Move foliage into its own module

diff --git a/Engine/Config/BaseEngine.ini b/Engine/Config/BaseEngine.ini
--- a/Engine/Config/BaseEngine.ini
+++ b/Engine/Config/BaseEngine.ini
@@ -1130,7 +1130,14 @@ bUseStreamingPause=false
 +TaggedPropertyRedirects=(ClassName="SplineMeshComponent",OldPropertyName="SplineXDir",NewPropertyName="SplineUpDir")
 
 ;Foliage refactor
-+ActiveClassRedirects=(OldClassName="InstancedFoliageSettings", NewClassName="FoliageType_InstancedStaticMesh")
++ActiveClassRedirects=(OldClassName="InstancedFoliageSettings", NewClassName="/Script/Foliage.FoliageType_InstancedStaticMesh")
++ActiveClassRedirects=(OldClassName="FoliageType", NewClassName="/Script/Foliage.FoliageType")
++ActiveClassRedirects=(OldClassName="FoliageType_InstancedStaticMesh", NewClassName="/Script/Foliage.FoliageType_InstancedStaticMesh")
++ActiveClassRedirects=(OldClassName="InstancedFoliageActor", NewClassName="/Script/Foliage.InstancedFoliageActor")
++ActiveClassRedirects=(OldClassName="InteractiveFoliageComponent", NewClassName="/Script/Foliage.InteractiveFoliageComponent")
++ActiveClassRedirects=(OldClassName="FoliageVertexColorMask", NewClassName="/Script/Foliage.FoliageVertexColorMask")
+
+
 
 ;Removal of EmitterSpawnable
 +ActiveClassRedirects=(OldClassName="EmitterSpawnable", NewClassName="Emitter")
diff --git a/Engine/Source/Developer/AssetTools/AssetTools.Build.cs b/Engine/Source/Developer/AssetTools/AssetTools.Build.cs
--- a/Engine/Source/Developer/AssetTools/AssetTools.Build.cs
+++ b/Engine/Source/Developer/AssetTools/AssetTools.Build.cs
@@ -23,6 +23,7 @@ public class AssetTools : ModuleRules
 				"UnrealEd",
 				"Kismet",
 				"Landscape",
+                "Foliage",
 			}
 		);
 
diff --git a/Engine/Source/Developer/AssetTools/Private/AssetTypeActions/AssetTypeActions_InstancedFoliageSettings.cpp b/Engine/Source/Developer/AssetTools/Private/AssetTypeActions/AssetTypeActions_InstancedFoliageSettings.cpp
--- a/Engine/Source/Developer/AssetTools/Private/AssetTypeActions/AssetTypeActions_InstancedFoliageSettings.cpp
+++ b/Engine/Source/Developer/AssetTools/Private/AssetTypeActions/AssetTypeActions_InstancedFoliageSettings.cpp
@@ -1,7 +1,7 @@
 // Copyright 1998-2015 Epic Games, Inc. All Rights Reserved.
 
 #include "AssetToolsPrivatePCH.h"
-#include "Foliage/FoliageType_InstancedStaticMesh.h"
+#include "FoliageType_InstancedStaticMesh.h"
 
 UClass* FAssetTypeActions_InstancedFoliageSettings::GetSupportedClass() const
 {
diff --git a/Engine/Source/Editor/FoliageEdit/FoliageEdit.build.cs b/Engine/Source/Editor/FoliageEdit/FoliageEdit.build.cs
--- a/Engine/Source/Editor/FoliageEdit/FoliageEdit.build.cs
+++ b/Engine/Source/Editor/FoliageEdit/FoliageEdit.build.cs
@@ -20,6 +20,7 @@ public class FoliageEdit : ModuleRules
 				"LevelEditor",
 				"Landscape",
                 "PropertyEditor",
+                "Foliage",
 			}
 		);
 	}
diff --git a/Engine/Source/Editor/FoliageEdit/Private/FoliageEdMode.cpp b/Engine/Source/Editor/FoliageEdit/Private/FoliageEdMode.cpp
--- a/Engine/Source/Editor/FoliageEdit/Private/FoliageEdMode.cpp
+++ b/Engine/Source/Editor/FoliageEdit/Private/FoliageEdMode.cpp
@@ -18,8 +18,8 @@
 #include "Editor/LevelEditor/Public/SLevelViewport.h"
 
 // Classes
-#include "Foliage/InstancedFoliageActor.h"
-#include "Foliage/FoliageType.h"
+#include "InstancedFoliageActor.h"
+#include "FoliageType.h"
 #include "LandscapeProxy.h"
 #include "LandscapeComponent.h"
 #include "LandscapeHeightfieldCollisionComponent.h"
diff --git a/Engine/Source/Editor/FoliageEdit/Private/SFoliageEdit.cpp b/Engine/Source/Editor/FoliageEdit/Private/SFoliageEdit.cpp
--- a/Engine/Source/Editor/FoliageEdit/Private/SFoliageEdit.cpp
+++ b/Engine/Source/Editor/FoliageEdit/Private/SFoliageEdit.cpp
@@ -7,7 +7,7 @@
 #include "SFoliageEdit.h"
 #include "FoliageEditActions.h"
 #include "FoliageEdMode.h"
-#include "Foliage/FoliageType.h"
+#include "FoliageType.h"
 #include "Editor/UnrealEd/Public/AssetThumbnail.h"
 #include "Editor/UnrealEd/Public/DragAndDrop/AssetDragDropOp.h"
 #include "Editor/UnrealEd/Public/AssetSelection.h"
diff --git a/Engine/Source/Editor/FoliageEdit/Private/SFoliageEditMeshDisplayItem.cpp b/Engine/Source/Editor/FoliageEdit/Private/SFoliageEditMeshDisplayItem.cpp
--- a/Engine/Source/Editor/FoliageEdit/Private/SFoliageEditMeshDisplayItem.cpp
+++ b/Engine/Source/Editor/FoliageEdit/Private/SFoliageEditMeshDisplayItem.cpp
@@ -5,7 +5,7 @@
 #include "SFoliageEdit.h"
 #include "FoliageEditActions.h"
 #include "FoliageEdMode.h"
-#include "Foliage/FoliageType.h"
+#include "FoliageType.h"
 #include "Editor/UnrealEd/Public/AssetThumbnail.h"
 #include "Dialogs/DlgPickAssetPath.h"
 #include "Editor/PropertyEditor/Public/PropertyEditorModule.h"
diff --git a/Engine/Source/Editor/Kismet/Private/InstancedStaticMeshSCSEditorCustomization.cpp b/Engine/Source/Editor/Kismet/Private/InstancedStaticMeshSCSEditorCustomization.cpp
--- a/Engine/Source/Editor/Kismet/Private/InstancedStaticMeshSCSEditorCustomization.cpp
+++ b/Engine/Source/Editor/Kismet/Private/InstancedStaticMeshSCSEditorCustomization.cpp
@@ -3,7 +3,6 @@
 #include "BlueprintEditorPrivatePCH.h"
 #include "InstancedStaticMeshSCSEditorCustomization.h"
 #include "BlueprintEditorModule.h"
-#include "InstancedFoliage.h"
 #include "Components/InstancedStaticMeshComponent.h"
 
 TSharedRef<ISCSEditorCustomization> FInstancedStaticMeshSCSEditorCustomization::MakeInstance(TSharedRef< IBlueprintEditor > InBlueprintEditor)
diff --git a/Engine/Source/Editor/Kismet/Private/SCSEditorViewportClient.cpp b/Engine/Source/Editor/Kismet/Private/SCSEditorViewportClient.cpp
--- a/Engine/Source/Editor/Kismet/Private/SCSEditorViewportClient.cpp
+++ b/Engine/Source/Editor/Kismet/Private/SCSEditorViewportClient.cpp
@@ -15,7 +15,6 @@
 #include "Editor/UnrealEd/Public/Kismet2/ComponentEditorUtils.h"
 #include "ISCSEditorCustomization.h"
 #include "ComponentVisualizer.h"
-#include "InstancedFoliage.h"
 #include "Engine/SimpleConstructionScript.h"
 #include "CanvasTypes.h"
 #include "EngineUtils.h"
diff --git a/Engine/Source/Editor/LandscapeEditor/LandscapeEditor.Build.cs b/Engine/Source/Editor/LandscapeEditor/LandscapeEditor.Build.cs
--- a/Engine/Source/Editor/LandscapeEditor/LandscapeEditor.Build.cs
+++ b/Engine/Source/Editor/LandscapeEditor/LandscapeEditor.Build.cs
@@ -21,6 +21,7 @@ public class LandscapeEditor : ModuleRules
 				"PropertyEditor",
 				"ImageWrapper",
                 "EditorWidgets",
+                "Foliage",
 			}
 			);
 
diff --git a/Engine/Source/Editor/LandscapeEditor/Private/LandscapeEdMode.cpp b/Engine/Source/Editor/LandscapeEditor/Private/LandscapeEdMode.cpp
--- a/Engine/Source/Editor/LandscapeEditor/Private/LandscapeEdMode.cpp
+++ b/Engine/Source/Editor/LandscapeEditor/Private/LandscapeEdMode.cpp
@@ -30,7 +30,7 @@
 #include "LandscapeHeightfieldCollisionComponent.h"
 #include "LandscapeMaterialInstanceConstant.h"
 #include "LandscapeSplinesComponent.h"
-#include "Foliage/InstancedFoliageActor.h"
+#include "InstancedFoliageActor.h"
 #include "ComponentReregisterContext.h"
 #include "Engine/Selection.h"
 #include "LandscapeGizmoActiveActor.h"
diff --git a/Engine/Source/Editor/LandscapeEditor/Private/LandscapeEdModeComponentTools.cpp b/Engine/Source/Editor/LandscapeEditor/Private/LandscapeEdModeComponentTools.cpp
--- a/Engine/Source/Editor/LandscapeEditor/Private/LandscapeEdModeComponentTools.cpp
+++ b/Engine/Source/Editor/LandscapeEditor/Private/LandscapeEdModeComponentTools.cpp
@@ -13,7 +13,7 @@
 #include "LandscapeEditorModule.h"
 #include "Editor/PropertyEditor/Public/PropertyEditorModule.h"
 #include "LandscapeEdModeTools.h"
-#include "Foliage/InstancedFoliageActor.h"
+#include "InstancedFoliageActor.h"
 #include "ComponentReregisterContext.h"
 #include "PhysicalMaterials/PhysicalMaterial.h"
 
diff --git a/Engine/Source/Editor/LandscapeEditor/Private/LandscapeEdModeTools.h b/Engine/Source/Editor/LandscapeEditor/Private/LandscapeEdModeTools.h
--- a/Engine/Source/Editor/LandscapeEditor/Private/LandscapeEdModeTools.h
+++ b/Engine/Source/Editor/LandscapeEditor/Private/LandscapeEdModeTools.h
@@ -4,7 +4,7 @@
 
 #include "Landscape.h"
 #include "LandscapeHeightfieldCollisionComponent.h"
-#include "Foliage/InstancedFoliageActor.h"
+#include "InstancedFoliageActor.h"
 #include "AI/Navigation/NavigationSystem.h"
 
 //
diff --git a/Engine/Source/Editor/UnrealEd/Private/Editor.cpp b/Engine/Source/Editor/UnrealEd/Private/Editor.cpp
--- a/Engine/Source/Editor/UnrealEd/Private/Editor.cpp
+++ b/Engine/Source/Editor/UnrealEd/Private/Editor.cpp
@@ -3,7 +3,7 @@
 
 #include "UnrealEd.h"
 #include "Matinee/MatineeActor.h"
-#include "Engine/InteractiveFoliageActor.h"
+#include "InteractiveFoliageActor.h"
 #include "Animation/SkeletalMeshActor.h"
 #include "Animation/VertexAnim/VertexAnimation.h"
 #include "Engine/WorldComposition.h"
diff --git a/Engine/Source/Editor/UnrealEd/Private/EditorActor.cpp b/Engine/Source/Editor/UnrealEd/Private/EditorActor.cpp
--- a/Engine/Source/Editor/UnrealEd/Private/EditorActor.cpp
+++ b/Engine/Source/Editor/UnrealEd/Private/EditorActor.cpp
@@ -13,7 +13,7 @@
 #include "Editor/GeometryMode/Public/GeometryEdMode.h"
 #include "Editor/GeometryMode/Public/EditorGeometry.h"
 #include "ActorEditorUtils.h"
-#include "Foliage/InstancedFoliageActor.h"
+#include "InstancedFoliageActor.h"
 #include "Animation/SkeletalMeshActor.h"
 #include "Particles/Emitter.h"
 #include "Particles/ParticleSystemComponent.h"
diff --git a/Engine/Source/Editor/UnrealEd/Private/EditorExporters.cpp b/Engine/Source/Editor/UnrealEd/Private/EditorExporters.cpp
--- a/Engine/Source/Editor/UnrealEd/Private/EditorExporters.cpp
+++ b/Engine/Source/Editor/UnrealEd/Private/EditorExporters.cpp
@@ -16,12 +16,12 @@
 #include "RawMesh.h"
 #include "MaterialExportUtils.h"
 #include "ImageUtils.h"
-#include "Foliage/InstancedFoliageActor.h"
+#include "InstancedFoliageActor.h"
 #include "Landscape.h"
 #include "LandscapeComponent.h"
 #include "LandscapeInfo.h"
 #include "LandscapeDataAccess.h"
-#include "Foliage/FoliageType.h"
+#include "FoliageType.h"
 #include "UnrealExporter.h"
 #include "EngineModule.h"
 #include "InstancedFoliage.h"
diff --git a/Engine/Source/Editor/UnrealEd/Private/EditorObject.cpp b/Engine/Source/Editor/UnrealEd/Private/EditorObject.cpp
--- a/Engine/Source/Editor/UnrealEd/Private/EditorObject.cpp
+++ b/Engine/Source/Editor/UnrealEd/Private/EditorObject.cpp
@@ -9,8 +9,8 @@
 #include "BSPOps.h"
 #include "Kismet2/BlueprintEditorUtils.h"
 
-#include "Foliage/InstancedFoliageActor.h"
-#include "Foliage/FoliageType.h"
+#include "InstancedFoliageActor.h"
+#include "FoliageType.h"
 #include "InstancedFoliage.h"
 #include "Components/BrushComponent.h"
 #include "Engine/StaticMesh.h"
diff --git a/Engine/Source/Editor/UnrealEd/Private/EditorServer.cpp b/Engine/Source/Editor/UnrealEd/Private/EditorServer.cpp
--- a/Engine/Source/Editor/UnrealEd/Private/EditorServer.cpp
+++ b/Engine/Source/Editor/UnrealEd/Private/EditorServer.cpp
@@ -28,7 +28,7 @@
 #include "Matinee/InterpTrackAnimControl.h"
 #include "Matinee/InterpData.h"
 #include "Animation/SkeletalMeshActor.h"
-#include "Foliage/InstancedFoliageActor.h"
+#include "InstancedFoliageActor.h"
 
 #include "Editor/UnrealEd/Public/Kismet2/KismetEditorUtilities.h"
 #include "Editor/PropertyEditor/Public/PropertyEditorModule.h"
diff --git a/Engine/Source/Editor/UnrealEd/Private/Factories/ActorFactory.cpp b/Engine/Source/Editor/UnrealEd/Private/Factories/ActorFactory.cpp
--- a/Engine/Source/Editor/UnrealEd/Private/Factories/ActorFactory.cpp
+++ b/Engine/Source/Editor/UnrealEd/Private/Factories/ActorFactory.cpp
@@ -16,7 +16,7 @@ ActorFactory.cpp:
 #include "BSPOps.h"
 #include "Matinee/MatineeActor.h"
 #include "Matinee/InterpData.h"
-#include "Engine/InteractiveFoliageActor.h"
+#include "InteractiveFoliageActor.h"
 #include "Animation/SkeletalMeshActor.h"
 #include "Animation/VertexAnim/VertexAnimation.h"
 
diff --git a/Engine/Source/Editor/UnrealEd/Private/Factories/EditorFactories.cpp b/Engine/Source/Editor/UnrealEd/Private/Factories/EditorFactories.cpp
--- a/Engine/Source/Editor/UnrealEd/Private/Factories/EditorFactories.cpp
+++ b/Engine/Source/Editor/UnrealEd/Private/Factories/EditorFactories.cpp
@@ -42,7 +42,7 @@
 #include "EnumEditorUtils.h"
 #include "StructureEditorUtils.h"
 
-#include "Foliage/InstancedFoliageActor.h"
+#include "InstancedFoliageActor.h"
 
 #include "Particles/ParticleSystem.h"
 
diff --git a/Engine/Source/Editor/UnrealEd/UnrealEd.Build.cs b/Engine/Source/Editor/UnrealEd/UnrealEd.Build.cs
--- a/Engine/Source/Editor/UnrealEd/UnrealEd.Build.cs
+++ b/Engine/Source/Editor/UnrealEd/UnrealEd.Build.cs
@@ -50,7 +50,7 @@ public class UnrealEd : ModuleRules
                 "IntroTutorials",
                 "SuperSearch",
 				"OutputLog",
-				"Landscape"
+				"Landscape",
 			}
 		);
 
@@ -118,6 +118,7 @@ public class UnrealEd : ModuleRules
 				"Landscape",
 				"HeadMountedDisplay",
 				"MeshPaint",
+                "Foliage",
 			}
 		);
 
diff --git a/Engine/Source/Editor/WorldBrowser/Private/Tiles/WorldTileCollectionModel.cpp b/Engine/Source/Editor/WorldBrowser/Private/Tiles/WorldTileCollectionModel.cpp
--- a/Engine/Source/Editor/WorldBrowser/Private/Tiles/WorldTileCollectionModel.cpp
+++ b/Engine/Source/Editor/WorldBrowser/Private/Tiles/WorldTileCollectionModel.cpp
@@ -23,7 +23,7 @@
 #include "Landscape.h"
 #include "GameFramework/WorldSettings.h"
 #include "Engine/Light.h"
-#include "Foliage/InstancedFoliageActor.h"
+#include "InstancedFoliageActor.h"
 #include "LandscapeMeshProxyActor.h"
 #include "LandscapeMeshProxyComponent.h"
 
diff --git a/Engine/Source/Editor/WorldBrowser/WorldBrowser.Build.cs b/Engine/Source/Editor/WorldBrowser/WorldBrowser.Build.cs
--- a/Engine/Source/Editor/WorldBrowser/WorldBrowser.Build.cs
+++ b/Engine/Source/Editor/WorldBrowser/WorldBrowser.Build.cs
@@ -43,6 +43,7 @@ public class WorldBrowser : ModuleRules
                 "LandscapeEditor",
 				"Levels",
                 "ImageWrapper",
+                "Foliage",
             }
 		);
 
diff --git a/Engine/Source/Runtime/Engine/Classes/Components/InstancedStaticMeshComponent.h b/Engine/Source/Runtime/Engine/Classes/Components/InstancedStaticMeshComponent.h
--- a/Engine/Source/Runtime/Engine/Classes/Components/InstancedStaticMeshComponent.h
+++ b/Engine/Source/Runtime/Engine/Classes/Components/InstancedStaticMeshComponent.h
@@ -214,3 +214,19 @@ protected:
 	friend FInstancedShadowMap2D;
 };
 
+/** InstancedStaticMeshInstance hit proxy */
+struct HInstancedStaticMeshInstance : public HHitProxy
+{
+	UInstancedStaticMeshComponent* Component;
+	int32 InstanceIndex;
+
+	DECLARE_HIT_PROXY(ENGINE_API);
+	HInstancedStaticMeshInstance(UInstancedStaticMeshComponent* InComponent, int32 InInstanceIndex) : HHitProxy(HPP_World), Component(InComponent), InstanceIndex(InInstanceIndex) {}
+
+	virtual void AddReferencedObjects(FReferenceCollector& Collector) override;
+
+	virtual EMouseCursor::Type GetMouseCursor()
+	{
+		return EMouseCursor::CardinalCross;
+	}
+};
diff --git a/Engine/Source/Runtime/Engine/Classes/Components/InteractiveFoliageComponent.h b/Engine/Source/Runtime/Engine/Classes/Components/InteractiveFoliageComponent.h
deleted file mode 100644
diff --git a/Engine/Source/Runtime/Engine/Classes/Engine/InteractiveFoliageActor.h b/Engine/Source/Runtime/Engine/Classes/Engine/InteractiveFoliageActor.h
deleted file mode 100644
diff --git a/Engine/Source/Runtime/Engine/Classes/Engine/StaticMesh.h b/Engine/Source/Runtime/Engine/Classes/Engine/StaticMesh.h
--- a/Engine/Source/Runtime/Engine/Classes/Engine/StaticMesh.h
+++ b/Engine/Source/Runtime/Engine/Classes/Engine/StaticMesh.h
@@ -371,9 +371,7 @@ class UStaticMesh : public UObject, public IInterface_CollisionDataProvider, pub
 	FString HighResSourceMeshName;
 
 #if WITH_EDITORONLY_DATA
-	/** Default settings when using this mesh for instanced foliage */
-	UPROPERTY(EditAnywhere, AdvancedDisplay, Instanced, Category=StaticMesh)
-	const UFoliageType_InstancedStaticMesh* FoliageDefaultSettings;
+
 
 	/** Importing data and options used for this mesh */
 	UPROPERTY(VisibleAnywhere, Instanced, Category=ImportSettings)
diff --git a/Engine/Source/Runtime/Engine/Classes/Foliage/FoliageType.h b/Engine/Source/Runtime/Engine/Classes/Foliage/FoliageType.h
deleted file mode 100644
diff --git a/Engine/Source/Runtime/Engine/Classes/Foliage/FoliageType_InstancedStaticMesh.h b/Engine/Source/Runtime/Engine/Classes/Foliage/FoliageType_InstancedStaticMesh.h
deleted file mode 100644
diff --git a/Engine/Source/Runtime/Engine/Classes/Foliage/InstancedFoliageActor.h b/Engine/Source/Runtime/Engine/Classes/Foliage/InstancedFoliageActor.h
deleted file mode 100644
diff --git a/Engine/Source/Runtime/Engine/Private/ActorEditor.cpp b/Engine/Source/Runtime/Engine/Private/ActorEditor.cpp
--- a/Engine/Source/Runtime/Engine/Private/ActorEditor.cpp
+++ b/Engine/Source/Runtime/Engine/Private/ActorEditor.cpp
@@ -7,7 +7,6 @@
 #include "UObjectToken.h"
 #include "LevelUtils.h"
 #include "MapErrors.h"
-#include "Foliage/InstancedFoliageActor.h"
 #include "Engine/LevelBounds.h"
 #include "Components/ChildActorComponent.h"
 
@@ -34,27 +33,7 @@ void AActor::PostEditChangeProperty(FPropertyChangedEvent& PropertyChangedEvent)
 	UProperty* PropertyThatChanged = PropertyChangedEvent.Property;
 	FName PropertyName = PropertyThatChanged != NULL ? PropertyThatChanged->GetFName() : NAME_None;
 	
-	bool bTransformationChanged = false;
-
-	if ( PropertyName==Name_RelativeLocation || PropertyName==Name_RelativeRotation || PropertyName==Name_RelativeScale3D )
-	{
-		bTransformationChanged = true;
-
-		if ( GIsEditor && !GetWorld()->IsPlayInEditor() )
-		{
-			AInstancedFoliageActor* IFA = AInstancedFoliageActor::GetInstancedFoliageActorForLevel(GetLevel());
-			if (IFA)
-			{
-				TInlineComponentArray<UActorComponent*> Components;
-				GetComponents(Components);
-
-				for ( int32 Idx = 0 ; Idx < Components.Num() ; ++Idx )
-				{
-					IFA->MoveInstancesForMovedComponent( Components[Idx] );
-				}
-			}
-		}
-	}
+	const bool bTransformationChanged = (PropertyName == Name_RelativeLocation || PropertyName == Name_RelativeRotation || PropertyName == Name_RelativeScale3D);
 
 	if ( ReregisterComponentsWhenModified() )
 	{
@@ -92,21 +71,6 @@ void AActor::PostEditMove(bool bFinished)
 
 	if ( bFinished )
 	{
-		if ( GIsEditor && !GetWorld()->IsPlayInEditor() && !FLevelUtils::IsMovingLevel())
-		{
-			AInstancedFoliageActor* IFA = AInstancedFoliageActor::GetInstancedFoliageActorForLevel(GetLevel());
-			if (IFA)
-			{
-				TInlineComponentArray<UActorComponent*> Components;
-				GetComponents(Components);
-
-				for ( int32 Idx = 0 ; Idx < Components.Num() ; ++Idx )
-				{
-					IFA->MoveInstancesForMovedComponent( Components[Idx] );
-				}
-			}
-		}
-
 		GetWorld()->bDoDelayedUpdateCullDistanceVolumes = true;
 		GetWorld()->bAreConstraintsDirty = true;
 
diff --git a/Engine/Source/Runtime/Engine/Private/Character.cpp b/Engine/Source/Runtime/Engine/Private/Character.cpp
--- a/Engine/Source/Runtime/Engine/Private/Character.cpp
+++ b/Engine/Source/Runtime/Engine/Private/Character.cpp
@@ -13,7 +13,6 @@
 #include "DisplayDebugHelpers.h"
 #include "Animation/AnimInstance.h"
 #include "Animation/AnimMontage.h"
-#include "Foliage/InstancedFoliageActor.h"
 #include "Components/CapsuleComponent.h"
 #include "GameFramework/DamageType.h"
 
@@ -412,12 +411,6 @@ namespace MovementBaseUtility
 			AActor* NewBaseOwner = NewBase->GetOwner();
 			if (NewBaseOwner)
 			{
-				// NOTE: TTP# 341962: Temp hack to fix issues with landing on Foliage actors with thousands of components.
-				if (Cast<AInstancedFoliageActor>(NewBaseOwner))
-				{
-					return;
-				}
-
 				if (NewBaseOwner->PrimaryActorTick.bCanEverTick)
 				{
 					BasedObjectTick.AddPrerequisite(NewBaseOwner, NewBaseOwner->PrimaryActorTick);
@@ -443,12 +436,6 @@ namespace MovementBaseUtility
 			AActor* OldBaseOwner = OldBase->GetOwner();
 			if (OldBaseOwner)
 			{
-				// NOTE: TTP# 341962: Temp hack to fix issues with landing on Foliage actors with thousands of components.
-				if (Cast<AInstancedFoliageActor>(OldBaseOwner))
-				{
-					return;
-				}
-
 				BasedObjectTick.RemovePrerequisite(OldBaseOwner, OldBaseOwner->PrimaryActorTick);
 
 				// @TODO: We need to find a more efficient way of finding all ticking components in an actor.
diff --git a/Engine/Source/Runtime/Engine/Private/EnginePrivate.h b/Engine/Source/Runtime/Engine/Private/EnginePrivate.h
--- a/Engine/Source/Runtime/Engine/Private/EnginePrivate.h
+++ b/Engine/Source/Runtime/Engine/Private/EnginePrivate.h
@@ -93,5 +93,4 @@ EnginePrivate.h: Unreal engine private header file.
 #include "UnrealEngine.h"					// Unreal engine helpers.
 #include "CanvasTypes.h"							// Canvas.
 #include "EngineUtils.h"
-#include "InstancedFoliage.h"				// Instanced foliage.
 #include "TimerManager.h"					// Game play timers
diff --git a/Engine/Source/Runtime/Engine/Private/FoliageComponent.cpp b/Engine/Source/Runtime/Engine/Private/FoliageComponent.cpp
deleted file mode 100644
diff --git a/Engine/Source/Runtime/Engine/Private/InstancedFoliage.cpp b/Engine/Source/Runtime/Engine/Private/InstancedFoliage.cpp
deleted file mode 100644
diff --git a/Engine/Source/Runtime/Engine/Private/InstancedStaticMesh.cpp b/Engine/Source/Runtime/Engine/Private/InstancedStaticMesh.cpp
--- a/Engine/Source/Runtime/Engine/Private/InstancedStaticMesh.cpp
+++ b/Engine/Source/Runtime/Engine/Private/InstancedStaticMesh.cpp
@@ -13,6 +13,13 @@
 
 const int32 InstancedStaticMeshMaxTexCoord = 8;
 
+IMPLEMENT_HIT_PROXY(HInstancedStaticMeshInstance, HHitProxy);
+
+/** InstancedStaticMeshInstance hit proxy */
+void HInstancedStaticMeshInstance::AddReferencedObjects(FReferenceCollector& Collector)
+{
+	Collector.AddReferencedObject(Component);
+}
 
 FStaticMeshInstanceBuffer::FStaticMeshInstanceBuffer(ERHIFeatureLevel::Type InFeatureLevel):
 	InstanceData(NULL)
diff --git a/Engine/Source/Runtime/Engine/Private/InstancedStaticMesh.h b/Engine/Source/Runtime/Engine/Private/InstancedStaticMesh.h
--- a/Engine/Source/Runtime/Engine/Private/InstancedStaticMesh.h
+++ b/Engine/Source/Runtime/Engine/Private/InstancedStaticMesh.h
@@ -11,7 +11,6 @@
 #include "ShaderParameterUtils.h"
 
 #include "Misc/UObjectToken.h"
-#include "Components/InteractiveFoliageComponent.h"
 #include "Components/SplineMeshComponent.h"
 #include "Components/ModelComponent.h"
 #include "Components/NiagaraComponent.h"
@@ -35,7 +34,6 @@
 #include "StaticMeshLight.h"
 #include "SpeedTreeWind.h"
 #include "ComponentInstanceDataCache.h"
-#include "InstancedFoliage.h"
 #include "VertexFactory.h"
 #include "LocalVertexFactory.h"
 
@@ -352,7 +350,6 @@ private:
 	FShaderParameter CPUInstanceInverseTransform;
 };
 
-
 /*-----------------------------------------------------------------------------
 	FInstancedStaticMeshRenderData
 -----------------------------------------------------------------------------*/
diff --git a/Engine/Source/Runtime/Engine/Private/InteractiveFoliageActor.cpp b/Engine/Source/Runtime/Engine/Private/InteractiveFoliageActor.cpp
deleted file mode 100644
diff --git a/Engine/Source/Runtime/Engine/Private/Level.cpp b/Engine/Source/Runtime/Engine/Private/Level.cpp
--- a/Engine/Source/Runtime/Engine/Private/Level.cpp
+++ b/Engine/Source/Runtime/Engine/Private/Level.cpp
@@ -26,7 +26,6 @@ Level.cpp: Level-related functions
 #include "LevelUtils.h"
 #include "TargetPlatform.h"
 #include "ContentStreaming.h"
-#include "Foliage/InstancedFoliageActor.h"
 #include "Engine/NavigationObjectBase.h"
 #include "Engine/ShadowMapTexture2D.h"
 #include "Components/ModelComponent.h"
@@ -1275,11 +1274,9 @@ void ULevel::BuildStreamingData(UTexture2D* UpdateSpecificTextureOnly/*=NULL*/)
 			if ( !bIsClassDefaultObject && Primitive->IsRegistered() )
 			{
 				const AActor* const Owner				= Primitive->GetOwner();
-				const bool bIsFoliage					= Owner && Owner->IsA(AInstancedFoliageActor::StaticClass()) && Primitive->IsA(UInstancedStaticMeshComponent::StaticClass()); 
 				const bool bIsStatic					= Owner == NULL 
 															|| Primitive->Mobility == EComponentMobility::Static 
-															|| Primitive->Mobility == EComponentMobility::Stationary
-															|| bIsFoliage; // treat Foliage components as static, regardless of mobility settings
+															|| Primitive->Mobility == EComponentMobility::Stationary;
 
 				TArray<FStreamingTexturePrimitiveInfo> PrimitiveStreamingTextures;
 
diff --git a/Engine/Source/Runtime/Engine/Private/LevelUtils.cpp b/Engine/Source/Runtime/Engine/Private/LevelUtils.cpp
--- a/Engine/Source/Runtime/Engine/Private/LevelUtils.cpp
+++ b/Engine/Source/Runtime/Engine/Private/LevelUtils.cpp
@@ -4,7 +4,6 @@
 #include "EnginePrivate.h"
 #include "LevelUtils.h"
 #include "EditorSupportDelegates.h"
-#include "Foliage/InstancedFoliageActor.h"
 
 #if WITH_EDITOR
 #include "ScopedTransaction.h"
@@ -387,11 +386,11 @@ void FLevelUtils::ApplyLevelTransform( ULevel* Level, const FTransform& LevelTra
 		}
 
 		// We need to transform foliage Editor data which is stored in world space
-		AInstancedFoliageActor* FoliageActor = AInstancedFoliageActor::GetInstancedFoliageActorForLevel(Level);
+		/*AInstancedFoliageActor* FoliageActor = AInstancedFoliageActor::GetInstancedFoliageActorForLevel(Level);
 		if (FoliageActor)
 		{
 			FoliageActor->ApplyLevelTransform(LevelTransform);
-		}
+		}*/
 #endif // WITH_EDITOR
 	}
 }
diff --git a/Engine/Source/Runtime/Engine/Public/InstancedFoliage.h b/Engine/Source/Runtime/Engine/Public/InstancedFoliage.h
deleted file mode 100644
diff --git a/Engine/Source/Runtime/Foliage/Foliage.Build.cs b/Engine/Source/Runtime/Foliage/Foliage.Build.cs
new file mode 100644
--- /dev/null
+++ b/Engine/Source/Runtime/Foliage/Foliage.Build.cs
@@ -0,0 +1,25 @@
+// Copyright 1998-2015 Epic Games, Inc. All Rights Reserved.
+
+using UnrealBuildTool;
+
+public class Foliage: ModuleRules
+{
+	public Foliage(TargetInfo Target)
+	{
+		PrivateDependencyModuleNames.AddRange(
+			new string[] {
+				"Core",
+				"CoreUObject",
+                "Engine",
+				"RenderCore",
+				"RHI"
+			}
+		);
+
+        PrivateIncludePaths.AddRange(
+            new string[] {
+                "Runtime/Foliage/Private"
+            })
+		;
+	}
+}
diff --git a/Engine/Source/Runtime/Foliage/Private/FoliageComponent.cpp b/Engine/Source/Runtime/Foliage/Private/FoliageComponent.cpp
new file mode 100644
--- /dev/null
+++ b/Engine/Source/Runtime/Foliage/Private/FoliageComponent.cpp
@@ -0,0 +1,240 @@
+// Copyright 1998-2015 Epic Games, Inc. All Rights Reserved.
+
+/*=============================================================================
+	FoliageComponent.cpp: Foliage rendering implementation.
+=============================================================================*/
+
+#include "FoliagePrivate.h"
+#include "InteractiveFoliageActor.h"
+#include "InteractiveFoliageComponent.h"
+#include "Components/CapsuleComponent.h"
+
+/** Scene proxy class for UInteractiveFoliageComponent. */
+class FInteractiveFoliageSceneProxy : public FStaticMeshSceneProxy
+{
+public:
+
+	FInteractiveFoliageSceneProxy(UInteractiveFoliageComponent* InComponent) :
+		FStaticMeshSceneProxy(InComponent),
+		FoliageImpluseDirection(0,0,0),
+		FoliageNormalizedRotationAxisAndAngle(0,0,1,0)
+	{}
+
+	/** Accessor used by the rendering thread when setting foliage parameters for rendering. */
+	virtual void GetFoliageParameters(FVector& OutFoliageImpluseDirection, FVector4& OutFoliageNormalizedRotationAxisAndAngle) const
+	{
+		OutFoliageImpluseDirection = FoliageImpluseDirection;
+		OutFoliageNormalizedRotationAxisAndAngle = FoliageNormalizedRotationAxisAndAngle;
+	}
+
+	/** Updates the scene proxy with new foliage parameters from the game thread. */
+	void UpdateParameters_GameThread(const FVector& NewFoliageImpluseDirection, const FVector4& NewFoliageNormalizedRotationAxisAndAngle)
+	{
+		checkSlow(IsInGameThread());
+		ENQUEUE_UNIQUE_RENDER_COMMAND_THREEPARAMETER(
+			UpdateFoliageParameters,
+			FInteractiveFoliageSceneProxy*,FoliageProxy,this,
+			FVector,NewFoliageImpluseDirection,NewFoliageImpluseDirection,
+			FVector4,NewFoliageNormalizedRotationAxisAndAngle,NewFoliageNormalizedRotationAxisAndAngle,
+			{
+				FoliageProxy->FoliageImpluseDirection = NewFoliageImpluseDirection;
+				FoliageProxy->FoliageNormalizedRotationAxisAndAngle = NewFoliageNormalizedRotationAxisAndAngle;
+			});
+	}
+
+protected:
+
+	FVector FoliageImpluseDirection;
+	FVector4 FoliageNormalizedRotationAxisAndAngle;
+};
+
+UInteractiveFoliageComponent::UInteractiveFoliageComponent(const FObjectInitializer& ObjectInitializer)
+	: Super(ObjectInitializer)
+{
+}
+
+FPrimitiveSceneProxy* UInteractiveFoliageComponent::CreateSceneProxy()
+{
+	// If a static mesh does not exist then this component cannot be added to the scene.
+	if(StaticMesh == NULL
+		|| StaticMesh->RenderData == NULL
+		|| StaticMesh->RenderData->LODResources.Num() == 0
+		|| StaticMesh->RenderData->LODResources[0].VertexBuffer.GetNumVertices() == 0)
+	{
+		return NULL;
+	}
+
+	// Store the foliage proxy so we can push updates to it during Tick
+	FoliageSceneProxy = new FInteractiveFoliageSceneProxy(this);
+	return FoliageSceneProxy;
+}
+
+void UInteractiveFoliageComponent::DestroyRenderState_Concurrent()
+{
+	Super::DestroyRenderState_Concurrent();
+	FoliageSceneProxy = NULL;
+}
+
+float AInteractiveFoliageActor::TakeDamage(float DamageAmount, FDamageEvent const& DamageEvent, AController* EventInstigator, AActor* DamageCauser)
+{
+	FHitResult Hit;
+	FVector ImpulseDir;
+	DamageEvent.GetBestHitInfo(this, (EventInstigator ? EventInstigator->GetPawn() : NULL), Hit, ImpulseDir);
+
+	// Discard the magnitude of the momentum and use Damage as the length instead
+	FVector DamageImpulse = ImpulseDir.GetSafeNormal() * DamageAmount * FoliageDamageImpulseScale;
+
+	// Apply force magnitude clamps
+	DamageImpulse.X = FMath::Clamp(DamageImpulse.X, -MaxDamageImpulse, MaxDamageImpulse);
+	DamageImpulse.Y = FMath::Clamp(DamageImpulse.Y, -MaxDamageImpulse, MaxDamageImpulse);
+	DamageImpulse.Z = FMath::Clamp(DamageImpulse.Z, -MaxDamageImpulse, MaxDamageImpulse);
+
+	FoliageForce += DamageImpulse;
+	
+	// Bring this actor out of stasis so that it gets ticked now that a force has been applied
+	SetActorTickEnabled(true);
+
+	return 0.f;
+}
+
+void AInteractiveFoliageActor::CapsuleTouched(AActor* Other, UPrimitiveComponent* OtherComp, int32 OtherBodyIndex, bool bFromSweep, const FHitResult& SweepResult)
+{
+	if (Other != NULL && OtherComp != NULL)
+	{
+		UCapsuleComponent* const TouchingActorCapsule = Cast<UCapsuleComponent>(OtherComp);
+		if (TouchingActorCapsule && CapsuleComponent)
+		{
+			const FVector CenterToTouching = FVector(TouchingActorCapsule->Bounds.Origin.X, TouchingActorCapsule->Bounds.Origin.Y, CapsuleComponent->Bounds.Origin.Z) - CapsuleComponent->Bounds.Origin;
+			// Keep track of the first position on the collision cylinder that the touching actor intersected
+			//@todo - need to handle multiple touching actors
+			TouchingActorEntryPosition = GetRootComponent()->Bounds.Origin + CenterToTouching.GetSafeNormal() * CapsuleComponent->GetScaledCapsuleRadius();
+		}
+		// Bring this actor out of stasis so that it gets ticked now that a force has been applied
+		SetActorTickEnabled(true);
+	}
+}
+
+void AInteractiveFoliageActor::SetupCollisionCylinder()
+{
+	if (GetStaticMeshComponent()->StaticMesh)
+	{
+		const FBoxSphereBounds MeshBounds = GetStaticMeshComponent()->StaticMesh->GetBounds();
+		const FVector Scale3D = GetStaticMeshComponent()->RelativeScale3D;
+		// Set the cylinder's radius based off of the static mesh's bounds radius
+		// CollisionRadius is in world space so apply the actor's scale
+		CapsuleComponent->SetCapsuleSize(MeshBounds.SphereRadius * .7f * FMath::Max(Scale3D.X, Scale3D.Y), MeshBounds.BoxExtent.Z * Scale3D.Z);
+
+		// Ensure delegate is bound (just once)
+		CapsuleComponent->OnComponentBeginOverlap.RemoveDynamic(this, &AInteractiveFoliageActor::CapsuleTouched);
+		CapsuleComponent->OnComponentBeginOverlap.AddDynamic(this, &AInteractiveFoliageActor::CapsuleTouched);
+	}
+}
+
+void AInteractiveFoliageActor::Tick(float DeltaSeconds)
+{
+	UInteractiveFoliageComponent* const FoliageComponent = CastChecked<UInteractiveFoliageComponent>(GetStaticMeshComponent());
+	// Can only push updates to the scene proxy if we are being ticked while registered
+	// The proxy will be NULL on dedicated server
+	if (FoliageComponent->IsRegistered() && FoliageComponent->FoliageSceneProxy)
+	{
+		TArray<AActor*> Touching;
+		GetOverlappingActors(Touching);
+
+		for (int32 TouchingIndex = 0; TouchingIndex < Touching.Num(); TouchingIndex++)
+		{
+			AActor* TouchingActor = Touching[TouchingIndex];
+			if (TouchingActor != NULL && TouchingActor->GetRootComponent() != NULL)
+			{
+				const FVector TouchingActorPosition(TouchingActor->GetRootComponent()->Bounds.Origin.X, TouchingActor->GetRootComponent()->Bounds.Origin.Y, GetRootComponent()->Bounds.Origin.Z);
+				//DrawDebugLine(GetWorld(), TouchingActorPosition, GetRootComponent()->Bounds.Origin, 255, 255, 255, false);
+				// Operate on the touching actor's collision cylinder
+				//@todo - handle touching actors without collision cylinders
+				UCapsuleComponent* TouchingActorCapsule = Cast<UCapsuleComponent>(TouchingActor->GetRootComponent());
+				if (TouchingActorCapsule && CapsuleComponent)
+				{
+					FVector TouchingToCenter = GetRootComponent()->Bounds.Origin - TouchingActorPosition;
+					// Force the simulated position to be in the XY plane for simplicity
+					TouchingToCenter.Z = 0;
+					// Position on the collision cylinder mirrored across the cylinder's center from the position that the touching actor entered
+					const FVector OppositeTouchingEntryPosition = GetRootComponent()->Bounds.Origin + GetRootComponent()->Bounds.Origin - TouchingActorEntryPosition;
+
+					// Project the touching actor's center onto the vector from where it first entered to OppositeTouchingEntryPosition
+					// This results in the same directional force being applied for the duration of the other actor touching this foliage actor,
+					// Which prevents strange movement that results from just comparing cylinder centers.
+					const FVector ProjectedTouchingActorPosition = (TouchingActorPosition - OppositeTouchingEntryPosition).ProjectOnTo(TouchingActorEntryPosition - OppositeTouchingEntryPosition) + OppositeTouchingEntryPosition;
+					// Find the furthest position on the cylinder of the touching actor from OppositeTouchingEntryPosition
+					const FVector TouchingActorFurthestPosition = ProjectedTouchingActorPosition + (TouchingActorEntryPosition - OppositeTouchingEntryPosition).GetSafeNormal() * TouchingActorCapsule->GetScaledCapsuleRadius();
+					// Construct the impulse as the distance between the furthest cylinder positions minus the two cylinder's diameters
+					const FVector ImpulseDirection = 
+						- (OppositeTouchingEntryPosition - TouchingActorFurthestPosition 
+						- (OppositeTouchingEntryPosition - TouchingActorFurthestPosition).GetSafeNormal() * 2.0f * (TouchingActorCapsule->GetScaledCapsuleRadius() + CapsuleComponent->GetScaledCapsuleRadius()));
+
+					//DrawDebugLine(GetWorld(), GetRootComponent()->Bounds.Origin + FVector(0,0,100), GetRootComponent()->Bounds.Origin + ImpulseDirection + FVector(0,0,100), 100, 255, 100, false);
+
+					// Scale and clamp the touch force
+					FVector Impulse = ImpulseDirection * FoliageTouchImpulseScale;
+					Impulse.X = FMath::Clamp(Impulse.X, -MaxTouchImpulse, MaxTouchImpulse);
+					Impulse.Y = FMath::Clamp(Impulse.Y, -MaxTouchImpulse, MaxTouchImpulse);
+					Impulse.Z = FMath::Clamp(Impulse.Z, -MaxTouchImpulse, MaxTouchImpulse);
+					FoliageForce += Impulse;
+				}
+			}
+		}
+
+		// Apply spring stiffness, which is the force that pushes the simulated particle back to the origin
+		FoliageForce += -FoliageStiffness * FoliagePosition;
+		// Apply spring quadratic stiffness, which increases in magnitude with the square of the distance to the origin
+		// This prevents the spring from being displaced too much by touch and damage forces
+		FoliageForce += -FoliageStiffnessQuadratic * FoliagePosition.SizeSquared() * FoliagePosition.GetSafeNormal();
+		// Apply spring damping, which is like air resistance and causes the spring to lose energy over time
+		FoliageForce += -FoliageDamping * FoliageVelocity;
+
+		FoliageForce.X = FMath::Clamp(FoliageForce.X, -MaxForce, MaxForce);
+		FoliageForce.Y = FMath::Clamp(FoliageForce.Y, -MaxForce, MaxForce);
+		FoliageForce.Z = FMath::Clamp(FoliageForce.Z, -MaxForce, MaxForce);
+
+		FoliageVelocity += FoliageForce * DeltaSeconds;
+		FoliageForce = FVector::ZeroVector;
+
+		const float MaxVelocity = 1000.0f;
+		FoliageVelocity.X = FMath::Clamp(FoliageVelocity.X, -MaxVelocity, MaxVelocity);
+		FoliageVelocity.Y = FMath::Clamp(FoliageVelocity.Y, -MaxVelocity, MaxVelocity);
+		FoliageVelocity.Z = FMath::Clamp(FoliageVelocity.Z, -MaxVelocity, MaxVelocity);
+
+		FoliagePosition += FoliageVelocity * DeltaSeconds;
+
+		//DrawDebugLine(GetWorld(), GetRootComponent()->Bounds.Origin + FVector(0,0,100), GetRootComponent()->Bounds.Origin + FVector(0,0,100) + FoliagePosition, 255, 100, 100, false);
+
+		//@todo - derive this height from the static mesh
+		const float IntersectionHeight = 100.0f;
+		// Calculate the rotation angle using Sin(Angle) = Opposite / Hypotenuse
+		const float RotationAngle = -FMath::Asin(FoliagePosition.Size() / IntersectionHeight);
+		// Use a rotation angle perpendicular to the impulse direction and the z axis
+		const FVector NormalizedRotationAxis = FoliagePosition.SizeSquared() > KINDA_SMALL_NUMBER ? 
+			(FoliagePosition ^ FVector(0,0,1)).GetSafeNormal() :
+			FVector(0,0,1);
+
+		// Propagate the new rotation axis and angle to the rendering thread
+		FoliageComponent->FoliageSceneProxy->UpdateParameters_GameThread(FoliagePosition, FVector4(NormalizedRotationAxis, RotationAngle));
+
+		if (FoliagePosition.SizeSquared() < FMath::Square(KINDA_SMALL_NUMBER * 10.0f)
+			&& FoliageVelocity.SizeSquared() < FMath::Square(KINDA_SMALL_NUMBER * 10.0f))
+		{
+			// Go into stasis (will no longer be ticked) if this actor's spring simulation has stabilized
+			SetActorTickEnabled(false);
+		}
+	}
+	Super::Tick(DeltaSeconds);
+}
+
+void AInteractiveFoliageActor::PostActorCreated()
+{
+	Super::PostActorCreated();
+	SetupCollisionCylinder();
+}
+
+void AInteractiveFoliageActor::PostLoad()
+{
+	Super::PostLoad();
+	SetupCollisionCylinder();
+}
diff --git a/Engine/Source/Runtime/Foliage/Private/FoliagePrivate.h b/Engine/Source/Runtime/Foliage/Private/FoliagePrivate.h
new file mode 100644
--- /dev/null
+++ b/Engine/Source/Runtime/Foliage/Private/FoliagePrivate.h
@@ -0,0 +1,7 @@
+// Copyright 1998-2015 Epic Games, Inc. All Rights Reserved.
+
+
+#pragma once
+#include "StaticMeshResources.h"
+#include "FoliageType_InstancedStaticMesh.h"
+#include "InstancedFoliageActor.h"
diff --git a/Engine/Source/Runtime/Foliage/Private/InstancedFoliage.cpp b/Engine/Source/Runtime/Foliage/Private/InstancedFoliage.cpp
new file mode 100644
--- /dev/null
+++ b/Engine/Source/Runtime/Foliage/Private/InstancedFoliage.cpp
@@ -0,0 +1,1721 @@
+// Copyright 1998-2015 Epic Games, Inc. All Rights Reserved.
+
+/*=============================================================================
+	InstancedFoliage.cpp: Instanced foliage implementation.
+=============================================================================*/
+
+#include "FoliagePrivate.h"
+#include "InstancedFoliage.h"
+#include "MessageLog.h"
+#include "UObjectToken.h"
+#include "MapErrors.h"
+#include "Components/ModelComponent.h"
+#include "Components/HierarchicalInstancedStaticMeshComponent.h"
+#include "Serialization/CustomVersion.h"
+
+#define LOCTEXT_NAMESPACE "InstancedFoliage"
+
+#define DO_FOLIAGE_CHECK			0			// whether to validate foliage data during editing.
+#define FOLIAGE_CHECK_TRANSFORM		0			// whether to compare transforms between render and painting data.
+
+
+DEFINE_LOG_CATEGORY_STATIC(LogInstancedFoliage, Log, All);
+
+
+// Custom serialization version for all packages containing Instance Foliage
+struct FFoliageCustomVersion
+{
+	enum Type
+	{
+		// Before any version changes were made in the plugin
+		BeforeCustomVersionWasAdded = 0,
+		// Converted to use HierarchicalInstancedStaticMeshComponent
+		FoliageUsingHierarchicalISMC = 1,
+		// Changed Component to not RF_Transactional
+		HierarchicalISMCNonTransactional = 2,
+		// -----<new versions can be added above this line>-------------------------------------------------
+		VersionPlusOne,
+		LatestVersion = VersionPlusOne - 1
+	};
+
+	// The GUID for this custom version number
+	const static FGuid GUID;
+
+private:
+	FFoliageCustomVersion() {}
+};
+
+const FGuid FFoliageCustomVersion::GUID(0x430C4D19, 0x71544970, 0x87699B69, 0xDF90B0E5);
+// Register the custom version with core
+FCustomVersionRegistration GRegisterFoliageCustomVersion(FFoliageCustomVersion::GUID, FFoliageCustomVersion::LatestVersion, TEXT("FoliageVer"));
+
+
+//
+// Serializers for struct data
+//
+
+FArchive& operator<<(FArchive& Ar, FFoliageInstance& Instance)
+{
+	Ar << Instance.Base;
+	Ar << Instance.Location;
+	Ar << Instance.Rotation;
+	Ar << Instance.DrawScale3D;
+
+	if (Ar.CustomVer(FFoliageCustomVersion::GUID) < FFoliageCustomVersion::FoliageUsingHierarchicalISMC)
+	{
+		int32 OldClusterIndex;
+		Ar << OldClusterIndex;
+		Ar << Instance.PreAlignRotation;
+		Ar << Instance.Flags;
+
+		if (OldClusterIndex == INDEX_NONE)
+		{
+			// When converting, we need to skip over any instance that was previously deleted but still in the Instances array.
+			Instance.Flags |= FOLIAGE_InstanceDeleted;
+		}
+	}
+	else
+	{
+		Ar << Instance.PreAlignRotation;
+		Ar << Instance.Flags;
+	}
+	
+	Ar << Instance.ZOffset;
+
+	return Ar;
+}
+
+/**
+*	FFoliageInstanceCluster_Deprecated
+*/
+struct FFoliageInstanceCluster_Deprecated
+{
+	UInstancedStaticMeshComponent* ClusterComponent;
+	FBoxSphereBounds Bounds;
+
+#if WITH_EDITORONLY_DATA
+	TArray<int32> InstanceIndices;	// index into editor editor Instances array
+#endif
+
+	friend FArchive& operator<<(FArchive& Ar, FFoliageInstanceCluster_Deprecated& OldCluster)
+	{
+		check(Ar.CustomVer(FFoliageCustomVersion::GUID) < FFoliageCustomVersion::FoliageUsingHierarchicalISMC);
+
+		Ar << OldCluster.Bounds;
+		Ar << OldCluster.ClusterComponent;
+
+#if WITH_EDITORONLY_DATA
+		if (!Ar.ArIsFilterEditorOnly ||
+			Ar.UE4Ver() < VER_UE4_FOLIAGE_SETTINGS_TYPE)
+		{
+			Ar << OldCluster.InstanceIndices;
+		}
+#endif
+
+		return Ar;
+	}
+};
+
+FArchive& operator<<(FArchive& Ar, FFoliageMeshInfo& MeshInfo)
+{
+	if (Ar.CustomVer(FFoliageCustomVersion::GUID) >= FFoliageCustomVersion::FoliageUsingHierarchicalISMC)
+	{
+		Ar << MeshInfo.Component;
+	}
+	else
+	{
+		TArray<FFoliageInstanceCluster_Deprecated> OldInstanceClusters;
+		Ar << OldInstanceClusters;
+	}
+	
+#if WITH_EDITORONLY_DATA
+	if ((!Ar.ArIsFilterEditorOnly || Ar.UE4Ver() < VER_UE4_FOLIAGE_SETTINGS_TYPE) && 
+		(!(Ar.GetPortFlags() & PPF_DuplicateForPIE)))
+	{
+		Ar << MeshInfo.Instances;
+	}
+
+	// Serialize the transient data for undo.
+	if (Ar.IsTransacting())
+	{
+		Ar << *MeshInfo.InstanceHash;
+		Ar << MeshInfo.ComponentHash;
+		Ar << MeshInfo.SelectedIndices;
+	}
+#endif
+
+	return Ar;
+}
+
+FArchive& operator<<(FArchive& Ar, FFoliageComponentHashInfo& ComponentHashInfo)
+{
+	return Ar << ComponentHashInfo.CachedLocation << ComponentHashInfo.CachedRotation << ComponentHashInfo.CachedDrawScale << ComponentHashInfo.Instances;
+}
+
+//
+// UFoliageType
+//
+
+UFoliageType::UFoliageType(const FObjectInitializer& ObjectInitializer)
+	: Super(ObjectInitializer)
+{
+	Density = 100.0f;
+	Radius = 0.0f;
+	AlignToNormal = true;
+	RandomYaw = true;
+	UniformScale = true;
+	ScaleMinX = 1.0f;
+	ScaleMinY = 1.0f;
+	ScaleMinZ = 1.0f;
+	ScaleMaxX = 1.0f;
+	ScaleMaxY = 1.0f;
+	ScaleMaxZ = 1.0f;
+	AlignMaxAngle = 0.0f;
+	RandomPitchAngle = 0.0f;
+	GroundSlope = 45.0f;
+	HeightMin = -262144.0f;
+	HeightMax = 262144.0f;
+	ZOffsetMin = 0.0f;
+	ZOffsetMax = 0.0f;
+	LandscapeLayer = NAME_None;
+	DisplayOrder = 0;
+	IsSelected = false;
+	ShowNothing = false;
+	ShowPaintSettings = true;
+	ShowInstanceSettings = false;
+	ReapplyDensityAmount = 1.0f;
+	CollisionWithWorld = false;
+	CollisionScale = FVector(0.9f, 0.9f, 0.9f);
+	VertexColorMask = FOLIAGEVERTEXCOLORMASK_Disabled;
+	VertexColorMaskThreshold = 0.5f;
+
+	bEnableStaticLighting = true;
+	CastShadow = true;
+	bCastDynamicShadow = true;
+	bCastStaticShadow = true;
+	bAffectDynamicIndirectLighting = false;
+	// Most of the high instance count foliage like grass causes performance problems with distance field lighting
+	bAffectDistanceFieldLighting = false;
+	bCastShadowAsTwoSided = false;
+	bReceivesDecals = false;
+
+	bOverrideLightMapRes = false;
+	OverriddenLightMapRes = 8;
+
+	BodyInstance.SetCollisionProfileName(UCollisionProfile::NoCollision_ProfileName);
+
+	/** Ecosystem settings*/
+	AverageSpreadDistance = 50;
+	SpreadVariance = 150;
+	bGrowsInShade = false;
+	SeedsPerStep = 3;
+	OverlapPriority = 0.f;
+	NumSteps = 3;
+	MinScale = 1.f;
+	MaxScale = 3.f;
+	ChangeCount = 0;
+	InitialSeedDensity = 1.f;
+	CollisionRadius = 100.f;
+	ShadeRadius = 100.f;
+	InitialMaxAge = 0.f;
+	MaxAge = 10.f;
+
+	FRichCurve* Curve = ScaleCurve.GetRichCurve();
+	Curve->AddKey(0.f, 0.f);
+	Curve->AddKey(1.f, 1.f);
+}
+
+
+UFoliageType_InstancedStaticMesh::UFoliageType_InstancedStaticMesh(const FObjectInitializer& ObjectInitializer)
+	: Super(ObjectInitializer)
+{
+	Mesh = nullptr;
+}
+
+
+float UFoliageType::GetMaxRadius() const
+{
+	return FMath::Max(CollisionRadius, ShadeRadius);
+}
+
+float UFoliageType::GetScaleForAge(const float Age) const
+{
+	const FRichCurve* Curve = ScaleCurve.GetRichCurveConst();
+	const float Time = FMath::Clamp(MaxAge == 0 ? 1.f : Age / MaxAge, 0.f, 1.f);
+	const float Scale = Curve->Eval(Time);
+	return MinScale + (MaxScale - MinScale) * Scale;
+}
+
+float UFoliageType::GetInitAge(FRandomStream& RandomStream) const
+{
+	return RandomStream.FRandRange(0, InitialMaxAge);
+}
+
+float UFoliageType::GetNextAge(const float CurrentAge, const int32 NumSteps) const
+{
+	float NewAge = CurrentAge;
+	for (int32 Count = 0; Count < NumSteps; ++Count)
+	{
+		const float GrowAge = NewAge + 1;
+		if (GrowAge <= MaxAge)
+		{
+			NewAge = GrowAge;
+		}
+		else
+		{
+			break;
+		}
+	}
+
+	return NewAge;
+}
+
+#if WITH_EDITOR
+void UFoliageType::PostEditChangeProperty(struct FPropertyChangedEvent& PropertyChangedEvent)
+{
+	Super::PostEditChangeProperty(PropertyChangedEvent);
+
+	// Ensure that OverriddenLightMapRes is a factor of 4
+	OverriddenLightMapRes = OverriddenLightMapRes > 4 ? OverriddenLightMapRes + 3 & ~3 : 4;
+	++ChangeCount;
+}
+#endif
+
+//
+// FFoliageMeshInfo
+//
+FFoliageMeshInfo::FFoliageMeshInfo()
+	: Component(nullptr)
+#if WITH_EDITOR
+	, InstanceHash(GIsEditor ? new FFoliageInstanceHash() : nullptr)
+#endif
+{ }
+
+
+#if WITH_EDITOR
+
+void FFoliageMeshInfo::CheckValid()
+{
+#if DO_FOLIAGE_CHECK
+	int32 ClusterTotal = 0;
+	int32 ComponentTotal = 0;
+
+	for (FFoliageInstanceCluster& Cluster : InstanceClusters)
+	{
+		check(Cluster.ClusterComponent != nullptr);
+		ClusterTotal += Cluster.InstanceIndices.Num();
+		ComponentTotal += Cluster.ClusterComponent->PerInstanceSMData.Num();
+	}
+
+	check(ClusterTotal == ComponentTotal);
+
+	int32 FreeTotal = 0;
+	int32 InstanceTotal = 0;
+	for (int32 InstanceIdx = 0; InstanceIdx < Instances.Num(); InstanceIdx++)
+	{
+		if (Instances[InstanceIdx].ClusterIndex != -1)
+		{
+			InstanceTotal++;
+		}
+		else
+		{
+			FreeTotal++;
+		}
+	}
+
+	check( ClusterTotal == InstanceTotal );
+	check( FreeInstanceIndices.Num() == FreeTotal );
+
+	InstanceHash->CheckInstanceCount(InstanceTotal);
+
+	int32 ComponentHashTotal = 0;
+	for( TMap<UActorComponent*, FFoliageComponentHashInfo >::TConstIterator It(ComponentHash); It; ++It )
+	{
+		ComponentHashTotal += It.Value().Instances.Num();
+	}
+	check( ComponentHashTotal == InstanceTotal);
+
+#if FOLIAGE_CHECK_TRANSFORM
+	// Check transforms match up with editor data
+	int32 MismatchCount = 0;
+	for( int32 ClusterIdx=0;ClusterIdx<InstanceClusters.Num();ClusterIdx++ )
+	{
+		TArray<int32> Indices = InstanceClusters(ClusterIdx).InstanceIndices;
+		UInstancedStaticMeshComponent* Comp = InstanceClusters(ClusterIdx).ClusterComponent;
+		for( int32 InstIdx=0;InstIdx<Indices.Num();InstIdx++ )
+		{
+			int32 InstanceIdx = Indices(InstIdx);
+
+			FTransform InstanceToWorldEd = Instances(InstanceIdx).GetInstanceTransform();
+			FTransform InstanceToWorldCluster = Comp->PerInstanceSMData(InstIdx).Transform * Comp->GetComponentToWorld();
+
+			if( !InstanceToWorldEd.Equals(InstanceToWorldCluster) )
+			{
+				Comp->PerInstanceSMData(InstIdx).Transform = InstanceToWorldEd.ToMatrixWithScale();
+				MismatchCount++;
+			}
+		}
+	}
+
+	if( MismatchCount != 0 )
+	{
+		UE_LOG(LogInstancedFoliage, Log, TEXT("%s: transform mismatch: %d"), *InstanceClusters(0).ClusterComponent->StaticMesh->GetName(), MismatchCount);
+	}
+#endif
+
+#endif
+}
+
+void FFoliageMeshInfo::AddInstance(AInstancedFoliageActor* InIFA, UFoliageType* InSettings, const FFoliageInstance& InNewInstance)
+{
+	InIFA->Modify();
+
+	if (Component == nullptr)
+	{
+		Component = ConstructObject<UHierarchicalInstancedStaticMeshComponent>(UHierarchicalInstancedStaticMeshComponent::StaticClass(), InIFA, NAME_None, RF_Transactional);
+
+		Component->Mobility = InSettings->bEnableStaticLighting ? EComponentMobility::Static : EComponentMobility::Movable;
+
+		Component->StaticMesh = InSettings->GetStaticMesh();
+		Component->bSelectable = true;
+		Component->bHasPerInstanceHitProxies = true;
+		Component->InstancingRandomSeed = FMath::Rand();
+		Component->InstanceStartCullDistance = InSettings->StartCullDistance;
+		Component->InstanceEndCullDistance = InSettings->EndCullDistance;
+
+		Component->CastShadow = InSettings->CastShadow;
+		Component->bCastDynamicShadow = InSettings->bCastDynamicShadow;
+		Component->bCastStaticShadow = InSettings->bCastStaticShadow;
+		Component->bAffectDynamicIndirectLighting = InSettings->bAffectDynamicIndirectLighting;
+		Component->bAffectDistanceFieldLighting = InSettings->bAffectDistanceFieldLighting;
+		Component->bCastShadowAsTwoSided = InSettings->bCastShadowAsTwoSided;
+		Component->bReceivesDecals = InSettings->bReceivesDecals;
+		Component->bOverrideLightMapRes = InSettings->bOverrideLightMapRes;
+		Component->OverriddenLightMapRes = InSettings->OverriddenLightMapRes;
+
+		Component->BodyInstance.CopyBodyInstancePropertiesFrom(&InSettings->BodyInstance);
+
+		Component->AttachTo(InIFA->GetRootComponent());
+		Component->RegisterComponent();
+
+		// Use only instance translation as a component transform
+		Component->SetWorldTransform(InIFA->GetRootComponent()->ComponentToWorld);
+
+		// Add the new component to the transaction buffer so it will get destroyed on undo
+		Component->Modify();
+		// We don't want to track changes to instances later so we mark it as non-transactional
+		Component->ClearFlags(RF_Transactional);
+	}
+	else
+	{
+		Component->InvalidateLightingCache();
+	}
+
+	// Add the instance taking either a free slot or adding a new item.
+	int32 InstanceIndex = Instances.AddUninitialized();
+
+	FFoliageInstance& AddedInstance = Instances[InstanceIndex];
+	AddedInstance = InNewInstance;
+
+	// Add the instance to the hash
+	InstanceHash->InsertInstance(InNewInstance.Location, InstanceIndex);
+	FFoliageComponentHashInfo* ComponentHashInfo = ComponentHash.Find(InNewInstance.Base);
+	if (ComponentHashInfo == nullptr)
+	{
+		ComponentHashInfo = &ComponentHash.Add(InNewInstance.Base, FFoliageComponentHashInfo(InNewInstance.Base));
+	}
+	ComponentHashInfo->Instances.Add(InstanceIndex);
+	
+	// Calculate transform for the instance
+	FTransform InstanceToWorld = InNewInstance.GetInstanceWorldTransform();
+
+	// Add the instance to the component
+	Component->AddInstanceWorldSpace(InstanceToWorld);
+
+	// Update PrimitiveComponent's culling distance taking into account the radius of the bounds, as
+	// it is based on the center of the component's bounds.
+//	float CullDistance = InSettings->EndCullDistance > 0 ? (float)InSettings->EndCullDistance + BestCluster->Bounds.SphereRadius : 0.f;
+//	Component->LDMaxDrawDistance = CullDistance;
+//	Component->CachedMaxDrawDistance = CullDistance;
+
+	CheckValid();
+}
+
+void FFoliageMeshInfo::RemoveInstances(AInstancedFoliageActor* InIFA, const TArray<int32>& InInstancesToRemove)
+{
+	if (InInstancesToRemove.Num())
+	{
+		check(Component);
+		InIFA->Modify();
+
+		TSet<int32> InstancesToRemove;
+		for (int32 Instance : InInstancesToRemove)
+		{
+			InstancesToRemove.Add(Instance);
+		}
+
+		while(InstancesToRemove.Num())
+		{
+			// Get an item from the set for processing
+			auto It = InstancesToRemove.CreateConstIterator();
+			int32 InstanceIndex = *It;		
+			int32 InstanceIndexToRemove = InstanceIndex;
+
+			FFoliageInstance& Instance = Instances[InstanceIndex];
+
+			// remove from hash
+			InstanceHash->RemoveInstance(Instance.Location, InstanceIndex);
+		    FFoliageComponentHashInfo* ComponentHashInfo = ComponentHash.Find(Instance.Base);
+		    if (ComponentHashInfo)
+		    {
+			    ComponentHashInfo->Instances.Remove(InstanceIndex);
+			    if (ComponentHashInfo->Instances.Num() == 0)
+			    {
+				    // Remove the component from the component hash if this is the last instance.
+				    ComponentHash.Remove(Instance.Base);
+			    }
+		    }
+
+			// remove from the component
+			Component->RemoveInstance(InstanceIndex);
+
+			// Remove it from the selection.
+			SelectedIndices.Remove(InstanceIndex);
+
+			// remove from instances array
+			Instances.RemoveAtSwap(InstanceIndex);
+
+			// update hashes for swapped instance
+			if (InstanceIndex != Instances.Num() && Instances.Num())
+			{
+				// Instance hash
+				FFoliageInstance& SwappedInstance = Instances[InstanceIndex];
+				InstanceHash->RemoveInstance(SwappedInstance.Location, Instances.Num());
+				InstanceHash->InsertInstance(SwappedInstance.Location, InstanceIndex);
+
+				// Component hash
+				if (SwappedInstance.Base)
+				{
+					FFoliageComponentHashInfo* ComponentHashInfo = ComponentHash.Find(SwappedInstance.Base);
+					check(ComponentHashInfo);
+					ComponentHashInfo->Instances.Remove(Instances.Num());
+					ComponentHashInfo->Instances.Add(InstanceIndex);
+				}
+
+				// Selection
+				if (SelectedIndices.Contains(Instances.Num()))
+				{
+					SelectedIndices.Remove(Instances.Num());
+					SelectedIndices.Add(InstanceIndex);
+				}
+
+				// Removal list
+				if (InstancesToRemove.Contains(Instances.Num()))
+				{
+					// The item from the end of the array that we swapped in to InstanceIndex is also on the list to remove.
+					// Remove the item at the end of the array and leave InstanceIndex in the removal list.
+					InstanceIndexToRemove = Instances.Num();
+				}
+			}
+
+			// Remove the removed item from the removal list
+			InstancesToRemove.Remove(InstanceIndexToRemove);
+		}
+
+		InIFA->CheckSelection();
+
+		CheckValid();
+	}
+}
+
+void FFoliageMeshInfo::PreMoveInstances(AInstancedFoliageActor* InIFA, const TArray<int32>& InInstancesToMove)
+{
+	// Remove instances from the hash
+	for (TArray<int32>::TConstIterator It(InInstancesToMove); It; ++It)
+	{
+		int32 InstanceIndex = *It;
+		const FFoliageInstance& Instance = Instances[InstanceIndex];
+		InstanceHash->RemoveInstance(Instance.Location, InstanceIndex);
+	}
+}
+
+
+void FFoliageMeshInfo::PostUpdateInstances(AInstancedFoliageActor* InIFA, const TArray<int32>& InInstancesUpdated, bool bReAddToHash)
+{
+	if (InInstancesUpdated.Num())
+	{
+		check(Component);
+
+		for (TArray<int32>::TConstIterator It(InInstancesUpdated); It; ++It)
+		{
+			int32 InstanceIndex = *It;
+			const FFoliageInstance& Instance = Instances[InstanceIndex];
+
+			FTransform InstanceToWorld = Instance.GetInstanceWorldTransform();
+
+			Component->UpdateInstanceTransform(InstanceIndex, InstanceToWorld, true);
+
+			// Re-add instance to the hash if requested
+			if (bReAddToHash)
+			{
+				InstanceHash->InsertInstance(Instance.Location, InstanceIndex);
+			}
+		}
+
+		Component->InvalidateLightingCache();
+		Component->MarkRenderStateDirty();
+	}
+}
+
+void FFoliageMeshInfo::PostMoveInstances(AInstancedFoliageActor* InIFA, const TArray<int32>& InInstancesMoved)
+{
+	PostUpdateInstances(InIFA, InInstancesMoved, true);
+}
+
+void FFoliageMeshInfo::DuplicateInstances(AInstancedFoliageActor* InIFA, UFoliageType* InSettings, const TArray<int32>& InInstancesToDuplicate)
+{
+	for (int32 InstanceIndex : InInstancesToDuplicate)
+	{
+		const FFoliageInstance TempInstance = Instances[InstanceIndex];
+		AddInstance(InIFA, InSettings, TempInstance);
+	}
+}
+
+/* Get the number of placed instances */
+int32 FFoliageMeshInfo::GetInstanceCount() const
+{
+	return Instances.Num();
+}
+
+// Destroy existing clusters and reassign all instances to new clusters
+void FFoliageMeshInfo::ReallocateClusters(AInstancedFoliageActor* InIFA, UFoliageType* InSettings)
+{
+	// Detach all components
+	InIFA->UnregisterAllComponents();
+
+	if (Component != nullptr)
+	{
+		Component->bAutoRegister = false;
+		Component = nullptr;
+	}
+
+	// Remove everything
+	TArray<FFoliageInstance> OldInstances;
+	Exchange(Instances, OldInstances);
+	InstanceHash->Empty();
+	ComponentHash.Empty();
+	SelectedIndices.Empty();
+
+	// Re-add
+	for (FFoliageInstance& Instance : OldInstances)
+	{
+		if ((Instance.Flags & FOLIAGE_InstanceDeleted) == 0)
+		{
+			AddInstance(InIFA, InSettings, Instance);
+		}
+	}
+
+	InIFA->RegisterAllComponents();
+}
+
+void FFoliageMeshInfo::ReapplyInstancesToComponent()
+{
+	if (Component)
+	{
+		Component->UnregisterComponent();
+		Component->ClearInstances();
+
+		for (auto& Instance : Instances)
+		{
+			Component->AddInstanceWorldSpace(Instance.GetInstanceWorldTransform());
+		}
+
+		if (SelectedIndices.Num())
+		{
+			if (Component->SelectedInstances.Num() != Component->PerInstanceSMData.Num())
+			{
+				Component->SelectedInstances.Init(false, Component->PerInstanceSMData.Num());
+			}
+			for (int32 i : SelectedIndices)
+			{
+				Component->SelectedInstances[i] = true;
+			}
+		}
+
+		Component->RegisterComponent();
+	}
+}
+
+
+// Update settings in the clusters based on the current settings (eg culling distance, collision, ...)
+//void FFoliageMeshInfo::UpdateClusterSettings(AInstancedFoliageActor* InIFA)
+//{
+//	for (FFoliageInstanceCluster& Cluster : InstanceClusters)
+//	{
+//		UInstancedStaticMeshComponent* ClusterComponent = Cluster.ClusterComponent;
+//		ClusterComponent->MarkRenderStateDirty();
+//
+//		// Copy settings
+//		ClusterComponent->InstanceStartCullDistance = Settings->StartCullDistance;
+//		ClusterComponent->InstanceEndCullDistance = Settings->EndCullDistance;
+//
+//		// Update PrimitiveComponent's culling distance taking into account the radius of the bounds, as
+//		// it is based on the center of the component's bounds.
+//		float CullDistance = Settings->EndCullDistance > 0 ? (float)Settings->EndCullDistance + Cluster.Bounds.SphereRadius : 0.f;
+//		ClusterComponent->LDMaxDrawDistance = CullDistance;
+//		ClusterComponent->CachedMaxDrawDistance = CullDistance;
+//	}
+//
+//	InIFA->MarkComponentsRenderStateDirty();
+//}
+
+void FFoliageMeshInfo::GetInstancesInsideSphere(const FSphere& Sphere, TArray<int32>& OutInstances)
+{
+	auto TempInstances = InstanceHash->GetInstancesOverlappingBox(FBox::BuildAABB(Sphere.Center, FVector(Sphere.W)));
+	for (int32 Idx : TempInstances)
+	{
+		if (FSphere(Instances[Idx].Location, 0.f).IsInside(Sphere))
+		{
+			OutInstances.Add(Idx);
+		}
+	}
+}
+
+// Returns whether or not there is are any instances overlapping the sphere specified
+bool FFoliageMeshInfo::CheckForOverlappingSphere(const FSphere& Sphere)
+{
+	auto TempInstances = InstanceHash->GetInstancesOverlappingBox(FBox::BuildAABB(Sphere.Center, FVector(Sphere.W)));
+	for (int32 Idx : TempInstances)
+	{
+		if (FSphere(Instances[Idx].Location, 0.f).IsInside(Sphere))
+		{
+			return true;
+		}
+	}
+	return false;
+}
+
+// Returns whether or not there is are any instances overlapping the instance specified, excluding the set of instances provided
+bool FFoliageMeshInfo::CheckForOverlappingInstanceExcluding(int32 TestInstanceIdx, float Radius, TSet<int32>& ExcludeInstances)
+{
+	FSphere Sphere(Instances[TestInstanceIdx].Location, Radius);
+
+	auto TempInstances = InstanceHash->GetInstancesOverlappingBox(FBox::BuildAABB(Sphere.Center, FVector(Sphere.W)));
+	for (int32 Idx : TempInstances)
+	{
+		if (Idx != TestInstanceIdx && !ExcludeInstances.Contains(Idx) && FSphere(Instances[Idx].Location, 0.f).IsInside(Sphere))
+		{
+			return true;
+		}
+	}
+	return false;
+}
+
+void FFoliageMeshInfo::SelectInstances(AInstancedFoliageActor* InIFA, bool bSelect, TArray<int32>& InInstances)
+{
+	if (InInstances.Num())
+	{
+		check(Component);
+		InIFA->Modify();
+
+		if (bSelect)
+		{
+			// Apply selections to the component
+			Component->MarkRenderStateDirty();
+
+			if (Component->SelectedInstances.Num() != Component->PerInstanceSMData.Num())
+			{
+				Component->SelectedInstances.Init(false, Component->PerInstanceSMData.Num());
+			}
+
+			for (int32 i : InInstances)
+			{
+				SelectedIndices.Add(i);
+				Component->SelectedInstances[i] = true;
+			}
+		}
+		else
+		{
+			if (Component->SelectedInstances.Num())
+			{
+				Component->MarkRenderStateDirty();
+
+				for (int32 i : InInstances)
+				{
+					SelectedIndices.Remove(i);
+					Component->SelectedInstances[i] = false;
+				}
+			}
+		}
+	}
+}
+
+#endif	//WITH_EDITOR
+
+//
+// AInstancedFoliageActor
+//
+AInstancedFoliageActor::AInstancedFoliageActor(const FObjectInitializer& ObjectInitializer)
+	: Super(ObjectInitializer)
+{
+	USceneComponent* SceneComponent = CreateDefaultSubobject<USceneComponent>(TEXT("RootComponent0"));
+	RootComponent = SceneComponent;
+	RootComponent->Mobility = EComponentMobility::Static;
+	
+	SetActorEnableCollision(true);
+#if WITH_EDITORONLY_DATA
+	bListedInSceneOutliner = false;
+#endif // WITH_EDITORONLY_DATA
+	PrimaryActorTick.bCanEverTick = false;
+}
+
+#if WITH_EDITOR
+AInstancedFoliageActor* AInstancedFoliageActor::GetInstancedFoliageActorForCurrentLevel(UWorld* InWorld, bool bCreateIfNone)
+{
+	for (TActorIterator<AInstancedFoliageActor> It(InWorld); It; ++It)
+	{
+		AInstancedFoliageActor* InstancedFoliageActor = *It;
+		if (InstancedFoliageActor)
+		{
+			if (InstancedFoliageActor->GetLevel()->IsCurrentLevel() && !InstancedFoliageActor->IsPendingKill())
+			{
+				return InstancedFoliageActor;
+			}
+		}
+	}
+
+	return bCreateIfNone ? InWorld->SpawnActor<AInstancedFoliageActor>() : nullptr;
+}
+
+
+AInstancedFoliageActor* AInstancedFoliageActor::GetInstancedFoliageActorForLevel(ULevel* InLevel, bool bCreateIfNone /* = true */)
+{
+	if (!InLevel)
+	{
+		return nullptr;
+	}
+
+	for (int32 ActorIndex = 0; ActorIndex < InLevel->Actors.Num(); ++ActorIndex)
+	{
+		AInstancedFoliageActor* InstancedFoliageActor = Cast<AInstancedFoliageActor>(InLevel->Actors[ActorIndex]);
+		if (InstancedFoliageActor && !InstancedFoliageActor->IsPendingKill())
+		{
+			return InstancedFoliageActor;
+		}
+	}
+
+	if (bCreateIfNone)
+	{
+		FActorSpawnParameters SpawnParams;
+		SpawnParams.OverrideLevel = InLevel;
+
+		return InLevel->GetWorld()->SpawnActor<AInstancedFoliageActor>(SpawnParams);
+	}
+
+	return nullptr;
+}
+
+void AInstancedFoliageActor::MoveInstancesForMovedComponent(UActorComponent* InComponent)
+{
+	bool bUpdatedInstances = false;
+	bool bFirst = true;
+
+	for (auto& MeshPair : FoliageMeshes)
+	{
+		FFoliageMeshInfo& MeshInfo = *MeshPair.Value;
+		FFoliageComponentHashInfo* ComponentHashInfo = MeshInfo.ComponentHash.Find(InComponent);
+		if (ComponentHashInfo)
+		{
+			if (bFirst)
+			{
+				bFirst = false;
+				Modify();
+			}
+
+			FVector OldLocation = ComponentHashInfo->CachedLocation;
+			FRotator OldRotation = ComponentHashInfo->CachedRotation;
+			FVector OldDrawScale = ComponentHashInfo->CachedDrawScale;
+			ComponentHashInfo->UpdateLocationFromActor(InComponent);
+
+			for (int32 InstanceIndex : ComponentHashInfo->Instances)
+			{
+				FFoliageInstance& Instance = MeshInfo.Instances[InstanceIndex];
+
+				MeshInfo.InstanceHash->RemoveInstance(Instance.Location, InstanceIndex);
+
+				// Apply change
+				FMatrix DeltaTransform =
+					FRotationMatrix(Instance.Rotation) *
+					FTranslationMatrix(Instance.Location) *
+					FTranslationMatrix(-OldLocation) *
+					FInverseRotationMatrix(OldRotation) *
+					FScaleMatrix(ComponentHashInfo->CachedDrawScale / OldDrawScale) *
+					FRotationMatrix(ComponentHashInfo->CachedRotation) *
+					FTranslationMatrix(ComponentHashInfo->CachedLocation);
+
+				// Extract rotation and position
+				Instance.Location = DeltaTransform.GetOrigin();
+				Instance.Rotation = DeltaTransform.Rotator();
+
+				// Apply render data
+				check(MeshInfo.Component);
+				MeshInfo.Component->UpdateInstanceTransform(InstanceIndex, Instance.GetInstanceWorldTransform(), true);
+
+				// Re-add the new instance location to the hash
+				MeshInfo.InstanceHash->InsertInstance(Instance.Location, InstanceIndex);
+			}
+		}
+	}
+}
+
+void AInstancedFoliageActor::DeleteInstancesForComponent(UActorComponent* InComponent)
+{
+	for (auto& MeshPair : FoliageMeshes)
+	{
+		FFoliageMeshInfo& MeshInfo = *MeshPair.Value;
+		const FFoliageComponentHashInfo* ComponentHashInfo = MeshInfo.ComponentHash.Find(InComponent);
+		if (ComponentHashInfo)
+		{
+			MeshInfo.RemoveInstances(this, ComponentHashInfo->Instances.Array());
+		}
+	}
+}
+
+void AInstancedFoliageActor::DeleteInstancesForSpawner(UActorComponent* InComponent)
+{
+	for (auto& MeshPair : FoliageMeshes)
+	{
+		FFoliageMeshInfo& MeshInfo = *MeshPair.Value;
+		TArray<int32> InstancesToRemove;
+		for (int32 InstanceIdx = 0; InstanceIdx < MeshInfo.Instances.Num(); InstanceIdx++)
+		{
+			if (MeshInfo.Instances[InstanceIdx].Spawner == InComponent)
+			{
+				InstancesToRemove.Add(InstanceIdx);
+			}
+		}
+
+		if (InstancesToRemove.Num())
+		{
+			MeshInfo.RemoveInstances(this, InstancesToRemove);
+		}
+	}
+}
+
+void AInstancedFoliageActor::MoveInstancesForComponentToCurrentLevel(UActorComponent* InComponent)
+{
+	AInstancedFoliageActor* NewIFA = AInstancedFoliageActor::GetInstancedFoliageActorForCurrentLevel(InComponent->GetWorld());
+
+	for (auto& MeshPair : FoliageMeshes)
+	{
+		FFoliageMeshInfo& MeshInfo = *MeshPair.Value;
+		UFoliageType* FoliageType = MeshPair.Key;
+
+		// Duplicate the foliage type if it's not shared
+		if (FoliageType->GetOutermost() == GetOutermost())
+		{
+			FoliageType = (UFoliageType*)StaticDuplicateObject(FoliageType, NewIFA, nullptr, RF_AllFlags & ~(RF_Standalone | RF_Public));
+		}
+		
+		const FFoliageComponentHashInfo* ComponentHashInfo = MeshInfo.ComponentHash.Find(InComponent);
+		if (ComponentHashInfo)
+		{
+			FFoliageMeshInfo* NewMeshInfo = NewIFA->FindOrAddMesh(FoliageType);
+
+			// Add the foliage to the new level
+			for (int32 InstanceIndex : ComponentHashInfo->Instances)
+			{
+				NewMeshInfo->AddInstance(NewIFA, FoliageType, MeshInfo.Instances[InstanceIndex]);
+			}
+
+			// Remove from old level
+			MeshInfo.RemoveInstances(this, ComponentHashInfo->Instances.Array());
+		}
+	}
+}
+
+void AInstancedFoliageActor::MoveInstancesToNewComponent(UPrimitiveComponent* InOldComponent, UPrimitiveComponent* InNewComponent)
+{
+	AInstancedFoliageActor* NewIFA = AInstancedFoliageActor::GetInstancedFoliageActorForLevel(InNewComponent->GetTypedOuter<ULevel>());
+	if (!NewIFA)
+	{
+		NewIFA = InNewComponent->GetWorld()->SpawnActor<AInstancedFoliageActor>();
+	}
+	check(NewIFA);
+
+	for (auto& MeshPair : FoliageMeshes)
+	{
+		FFoliageMeshInfo& MeshInfo = *MeshPair.Value;
+		UFoliageType* FoliageType = MeshPair.Key;
+
+		// Duplicate the foliage type if it's not shared
+		if (FoliageType->GetOutermost() == GetOutermost())
+		{
+			FoliageType = (UFoliageType*)StaticDuplicateObject(FoliageType, NewIFA, nullptr, RF_AllFlags & ~(RF_Standalone | RF_Public));
+		}
+
+		const FFoliageComponentHashInfo* ComponentHashInfo = MeshInfo.ComponentHash.Find(InOldComponent);
+		if (ComponentHashInfo)
+		{
+			// For same level can just remap the instances, otherwise we have to do a more complex move
+			if (NewIFA == this)
+			{
+				FFoliageComponentHashInfo NewComponentHashInfo = MoveTemp(*ComponentHashInfo);
+				NewComponentHashInfo.UpdateLocationFromActor(InNewComponent);
+
+				// Update the instances
+				for (int32 InstanceIndex : NewComponentHashInfo.Instances)
+				{
+					MeshInfo.Instances[InstanceIndex].Base = InNewComponent;
+				}
+
+				// Update the hash
+				MeshInfo.ComponentHash.Add(InNewComponent, MoveTemp(NewComponentHashInfo));
+				MeshInfo.ComponentHash.Remove(InOldComponent);
+			}
+			else
+			{
+				FFoliageMeshInfo* NewMeshInfo = NewIFA->FindOrAddMesh(FoliageType);
+
+				// Add the foliage to the new level
+				for (int32 InstanceIndex : ComponentHashInfo->Instances)
+				{
+					FFoliageInstance NewInstance = MeshInfo.Instances[InstanceIndex];
+					NewInstance.Base = InNewComponent;
+					NewMeshInfo->AddInstance(NewIFA, FoliageType, NewInstance);
+				}
+
+				// Remove from old level
+				MeshInfo.RemoveInstances(this, ComponentHashInfo->Instances.Array());
+				check(!MeshInfo.ComponentHash.Contains(InOldComponent));
+			}
+		}
+	}
+}
+
+TMap<UFoliageType*, TArray<const FFoliageInstancePlacementInfo*>> AInstancedFoliageActor::GetInstancesForComponent(UActorComponent* InComponent)
+{
+	TMap<UFoliageType*, TArray<const FFoliageInstancePlacementInfo*>> Result;
+
+	for (auto& MeshPair : FoliageMeshes)
+	{
+		const FFoliageMeshInfo& MeshInfo = *MeshPair.Value;
+
+		const FFoliageComponentHashInfo* ComponentHashInfo = MeshInfo.ComponentHash.Find(InComponent);
+		if (ComponentHashInfo)
+		{
+			TArray<const FFoliageInstancePlacementInfo*>& Array = Result.Add(MeshPair.Key, TArray<const FFoliageInstancePlacementInfo*>());
+			Array.Empty(ComponentHashInfo->Instances.Num());
+
+			for (int32 InstanceIndex : ComponentHashInfo->Instances)
+			{
+				const FFoliageInstancePlacementInfo* Instance = &MeshInfo.Instances[InstanceIndex];
+				Array.Add(Instance);
+			}
+		}
+	}
+
+	return Result;
+}
+
+FFoliageMeshInfo* AInstancedFoliageActor::FindMesh(UFoliageType* InType)
+{
+	TUniqueObj<FFoliageMeshInfo>* MeshInfoEntry = FoliageMeshes.Find(InType);
+	FFoliageMeshInfo* MeshInfo = MeshInfoEntry ? &MeshInfoEntry->Get() : nullptr;
+	return MeshInfo;
+}
+
+FFoliageMeshInfo* AInstancedFoliageActor::FindOrAddMesh(UFoliageType* InType)
+{
+	TUniqueObj<FFoliageMeshInfo>* MeshInfoEntry = FoliageMeshes.Find(InType);
+	FFoliageMeshInfo* MeshInfo = MeshInfoEntry ? &MeshInfoEntry->Get() : AddMesh(InType);
+	return MeshInfo;
+}
+
+FFoliageMeshInfo* AInstancedFoliageActor::AddMesh(UStaticMesh* InMesh, UFoliageType** OutSettings)
+{
+	check(GetSettingsForMesh(InMesh) == nullptr);
+
+	MarkPackageDirty();
+
+	UFoliageType_InstancedStaticMesh* Settings = nullptr;
+#if WITH_EDITORONLY_DATA
+	/*if (InMesh->FoliageDefaultSettings)
+	{
+		// TODO: Can't we just use this directly?
+		Settings = DuplicateObject<UFoliageType_InstancedStaticMesh>(InMesh->FoliageDefaultSettings, this);
+	}
+	else*/
+#endif
+	{
+		Settings = ConstructObject<UFoliageType_InstancedStaticMesh>(UFoliageType_InstancedStaticMesh::StaticClass(), this);
+	}
+	Settings->Mesh = InMesh;
+
+	FFoliageMeshInfo* MeshInfo = AddMesh(Settings);
+
+	const FBoxSphereBounds MeshBounds = InMesh->GetBounds();
+
+	Settings->MeshBounds = MeshBounds;
+
+	// Make bottom only bound
+	FBox LowBound = MeshBounds.GetBox();
+	LowBound.Max.Z = LowBound.Min.Z + (LowBound.Max.Z - LowBound.Min.Z) * 0.1f;
+
+	float MinX = FLT_MAX, MaxX = FLT_MIN, MinY = FLT_MAX, MaxY = FLT_MIN;
+	Settings->LowBoundOriginRadius = FVector::ZeroVector;
+
+	if (InMesh->RenderData)
+	{
+		FPositionVertexBuffer& PositionVertexBuffer = InMesh->RenderData->LODResources[0].PositionVertexBuffer;
+		for (uint32 Index = 0; Index < PositionVertexBuffer.GetNumVertices(); ++Index)
+		{
+			const FVector& Pos = PositionVertexBuffer.VertexPosition(Index);
+			if (Pos.Z < LowBound.Max.Z)
+			{
+				MinX = FMath::Min(MinX, Pos.X);
+				MinY = FMath::Min(MinY, Pos.Y);
+				MaxX = FMath::Max(MaxX, Pos.X);
+				MaxY = FMath::Max(MaxY, Pos.Y);
+			}
+		}
+	}
+
+	Settings->LowBoundOriginRadius = FVector((MinX + MaxX), (MinY + MaxY), FMath::Sqrt(FMath::Square(MaxX - MinX) + FMath::Square(MaxY - MinY))) * 0.5f;
+
+	if (OutSettings)
+	{
+		*OutSettings = Settings;
+	}
+
+	return MeshInfo;
+}
+
+FFoliageMeshInfo* AInstancedFoliageActor::AddMesh(UFoliageType* InType)
+{
+	check(FoliageMeshes.Find(InType) == nullptr);
+
+	MarkPackageDirty();
+
+	if (InType->DisplayOrder == 0)
+	{
+		int32 MaxDisplayOrder = 0;
+		for (auto& MeshPair : FoliageMeshes)
+		{
+			if (MeshPair.Key->DisplayOrder > MaxDisplayOrder)
+			{
+				MaxDisplayOrder = MeshPair.Key->DisplayOrder;
+			}
+		}
+		InType->DisplayOrder = MaxDisplayOrder + 1;
+	}
+
+	FFoliageMeshInfo* MeshInfo = &*FoliageMeshes.Add(InType);
+	InType->IsSelected = true;
+
+	return MeshInfo;
+}
+
+void AInstancedFoliageActor::RemoveMesh(UFoliageType* InSettings)
+{
+	Modify();
+	MarkPackageDirty();
+	UnregisterAllComponents();
+
+	// Remove all components for this mesh from the Components array.
+	FFoliageMeshInfo* MeshInfo = FindMesh(InSettings);
+	if (MeshInfo && MeshInfo->Component)
+	{
+		MeshInfo->Component->bAutoRegister = false;
+	}
+
+	FoliageMeshes.Remove(InSettings);
+	RegisterAllComponents();
+	CheckSelection();
+}
+
+UFoliageType* AInstancedFoliageActor::GetSettingsForMesh(UStaticMesh* InMesh, FFoliageMeshInfo** OutMeshInfo)
+{
+	UFoliageType* Type = nullptr;
+	FFoliageMeshInfo* MeshInfo = nullptr;
+
+	for (auto& MeshPair : FoliageMeshes)
+	{
+		UFoliageType* Settings = MeshPair.Key;
+		if (Settings && Settings->GetStaticMesh() == InMesh)
+		{
+			Type = MeshPair.Key;
+			MeshInfo = &*MeshPair.Value;
+			break;
+		}
+	}
+
+	if (OutMeshInfo)
+	{
+		*OutMeshInfo = MeshInfo;
+	}
+	return Type;
+}
+
+void AInstancedFoliageActor::SelectInstance(UInstancedStaticMeshComponent* InComponent, int32 InInstanceIndex, bool bToggle)
+{
+	Modify();
+
+	// If we're not toggling, we need to first deselect everything else
+	if (!bToggle)
+	{
+		for (auto& MeshPair : FoliageMeshes)
+		{
+			FFoliageMeshInfo& MeshInfo = *MeshPair.Value;
+
+			if (MeshInfo.SelectedIndices.Num() > 0)
+			{
+				check(MeshInfo.Component);
+				if (MeshInfo.Component->SelectedInstances.Num() > 0)
+				{
+					MeshInfo.Component->SelectedInstances.Empty();
+					MeshInfo.Component->MarkRenderStateDirty();
+				}
+			}
+
+			MeshInfo.SelectedIndices.Empty();
+		}
+	}
+
+	if (InComponent)
+	{
+		UFoliageType* Type = nullptr;
+		FFoliageMeshInfo* MeshInfo = nullptr;
+
+		Type = GetSettingsForMesh(InComponent->StaticMesh, &MeshInfo);
+
+		if (MeshInfo)
+		{
+			if (InComponent == MeshInfo->Component)
+			{
+				bool bIsSelected = MeshInfo->SelectedIndices.Contains(InInstanceIndex);
+
+				// Deselect if it's already selected.
+				if (InInstanceIndex < InComponent->SelectedInstances.Num())
+				{
+					InComponent->SelectedInstances[InInstanceIndex] = false;
+					InComponent->MarkRenderStateDirty();
+				}
+
+				if (bIsSelected)
+				{
+					MeshInfo->SelectedIndices.Remove(InInstanceIndex);
+				}
+
+				if (bToggle && bIsSelected)
+				{
+					if (SelectedMesh == Type && MeshInfo->SelectedIndices.Num() == 0)
+					{
+						SelectedMesh = nullptr;
+					}
+				}
+				else
+				{
+					// Add the selection
+					if (InComponent->SelectedInstances.Num() < InComponent->PerInstanceSMData.Num())
+					{
+						InComponent->SelectedInstances.Init(false, InComponent->PerInstanceSMData.Num());
+					}
+					InComponent->SelectedInstances[InInstanceIndex] = true;
+					InComponent->MarkRenderStateDirty();
+
+					SelectedMesh = Type;
+					MeshInfo->SelectedIndices.Add(InInstanceIndex);
+				}
+			}
+		}
+	}
+
+	CheckSelection();
+
+}
+
+void AInstancedFoliageActor::PostEditUndo()
+{
+	Super::PostEditUndo();
+
+	FlushRenderingCommands();
+	for (auto& MeshPair : FoliageMeshes)
+	{
+		FFoliageMeshInfo& MeshInfo = *MeshPair.Value;
+		MeshInfo.ReapplyInstancesToComponent();
+	}
+}
+
+void AInstancedFoliageActor::PostEditChangeProperty(FPropertyChangedEvent& PropertyChangedEvent)
+{
+	Super::PostEditChangeProperty(PropertyChangedEvent);
+
+	if (RootComponent)
+	{
+		RootComponent->SetMobility(EComponentMobility::Static);	//ensure mobility is not changed
+	}
+}
+
+void AInstancedFoliageActor::ApplySelectionToComponents(bool bApply)
+{
+	bool bNeedsUpdate = false;
+
+	for (auto& MeshPair : FoliageMeshes)
+	{
+		FFoliageMeshInfo& MeshInfo = *MeshPair.Value;
+
+		if (bApply)
+		{
+			if (MeshInfo.SelectedIndices.Num() > 0)
+			{
+				check(MeshInfo.Component);
+
+				// Apply any selections in the component
+				MeshInfo.Component->SelectedInstances.Init(false, MeshInfo.Component->PerInstanceSMData.Num());
+				MeshInfo.Component->MarkRenderStateDirty();
+				for (int32 i : MeshInfo.SelectedIndices)
+				{
+					MeshInfo.Component->SelectedInstances[i] = true;
+				}
+			}
+		}
+		else		
+		{
+			if (MeshInfo.Component && MeshInfo.Component->SelectedInstances.Num() > 0)
+			{
+				// remove any selections in the component
+				MeshInfo.Component->SelectedInstances.Empty();
+				MeshInfo.Component->MarkRenderStateDirty();
+			}
+		}
+	}
+}
+
+void AInstancedFoliageActor::CheckSelection()
+{
+	// Check if we have to change the selection.
+	if (SelectedMesh != nullptr)
+	{
+		FFoliageMeshInfo* MeshInfo = FindMesh(SelectedMesh);
+		if (MeshInfo && MeshInfo->SelectedIndices.Num() > 0)
+		{
+			return;
+		}
+	}
+
+	SelectedMesh = nullptr;
+
+	// Try to find a new selection
+	for (auto& MeshPair : FoliageMeshes)
+	{
+		FFoliageMeshInfo& MeshInfo = *MeshPair.Value;
+		if (MeshInfo.SelectedIndices.Num() > 0)
+		{
+			SelectedMesh = MeshPair.Key;
+			return;
+		}
+	}
+}
+
+FVector AInstancedFoliageActor::GetSelectionLocation()
+{
+	FVector Result(0, 0, 0);
+
+	if (SelectedMesh != nullptr)
+	{
+		FFoliageMeshInfo* MeshInfo = FindMesh(SelectedMesh);
+		if (MeshInfo && MeshInfo->SelectedIndices.Num() > 0)
+		{
+			Result = MeshInfo->Instances[MeshInfo->SelectedIndices.Array()[0]].Location;
+		}
+	}
+
+	return Result;
+}
+
+void AInstancedFoliageActor::MapRebuild()
+{
+	// Map rebuild may have modified the BSP's ModelComponents and thrown the previous ones away.
+	// Most BSP-painted foliage is attached to a Brush's UModelComponent which persist across rebuilds,
+	// but any foliage attached directly to the level BSP's ModelComponents will need to try to find a new base.
+
+	TMap<UFoliageType*, TArray<FFoliageInstance>> NewInstances;
+	TArray<UModelComponent*> RemovedModelComponents;
+	UWorld* World = GetWorld();
+	check(World);
+
+	// For each foliage brush, represented by the mesh/info pair
+	for (auto& MeshPair : FoliageMeshes)
+	{
+		// each target component has some foliage instances
+		FFoliageMeshInfo const& MeshInfo = *MeshPair.Value;
+		UFoliageType* Settings = MeshPair.Key;
+		check(Settings);
+
+		for (auto& ComponentFoliagePair : MeshInfo.ComponentHash)
+		{
+			// BSP components are UModelComponents - they are the only ones we need to change
+			UModelComponent* TargetComponent = Cast<UModelComponent>(ComponentFoliagePair.Key);
+			// Check if it's part of a brush. We only need to fix up model components that are part of the level BSP.
+			if (TargetComponent && Cast<ABrush>(TargetComponent->GetOuter()) == nullptr)
+			{
+				// Delete its instances later
+				RemovedModelComponents.Add(TargetComponent);
+
+				FFoliageComponentHashInfo const& FoliageInfo = ComponentFoliagePair.Value;
+
+				// We have to test each instance to see if we can migrate it across
+				for (int32 InstanceIdx : FoliageInfo.Instances)
+				{
+					// Use a line test against the world. This is not very reliable as we don't know the original trace direction.
+					check(MeshInfo.Instances.IsValidIndex(InstanceIdx));
+					FFoliageInstance const& Instance = MeshInfo.Instances[InstanceIdx];
+
+					FFoliageInstance NewInstance = Instance;
+
+					FTransform InstanceToWorld = Instance.GetInstanceWorldTransform();
+					FVector Down(-FVector::UpVector);
+					FVector Start(InstanceToWorld.TransformPosition(FVector::UpVector));
+					FVector End(InstanceToWorld.TransformPosition(Down));
+
+					FHitResult Result;
+					bool bHit = World->LineTraceSingle(Result, Start, End, FCollisionQueryParams(true), FCollisionObjectQueryParams(ECC_WorldStatic));
+
+					if (bHit && Result.Component.IsValid() && Result.Component->IsA(UModelComponent::StaticClass()))
+					{
+						NewInstance.Base = CastChecked<UPrimitiveComponent>(Result.Component.Get());
+						NewInstances.FindOrAdd(Settings).Add(NewInstance);
+					}
+				}
+			}
+		}
+	}
+
+	// Remove all existing & broken instances & component references.
+	for (UModelComponent* Component : RemovedModelComponents)
+	{
+		DeleteInstancesForComponent(Component);
+	}
+
+	// And then finally add our new instances to the correct target components.
+	for (auto& NewInstancePair : NewInstances)
+	{
+		UFoliageType* Settings = NewInstancePair.Key;
+		check(Settings);
+		FFoliageMeshInfo& MeshInfo = *FindOrAddMesh(Settings);
+		for (FFoliageInstance& Instance : NewInstancePair.Value)
+		{
+			MeshInfo.AddInstance(this, Settings, Instance);
+		}
+	}
+}
+
+void AInstancedFoliageActor::ApplyLevelTransform(const FTransform& LevelTransform)
+{
+	// Apply transform to foliage editor only data
+	if (GIsEditor)
+	{
+		for (auto& MeshPair : FoliageMeshes)
+		{
+			FFoliageMeshInfo& MeshInfo = *MeshPair.Value;
+		
+			MeshInfo.InstanceHash->Empty();
+			for (int32 InstanceIdx = 0; InstanceIdx < MeshInfo.Instances.Num(); InstanceIdx++)
+			{
+				FFoliageInstance& Instance = MeshInfo.Instances[InstanceIdx];
+				FTransform NewTransform = Instance.GetInstanceWorldTransform() * LevelTransform;
+			
+				Instance.Location		= NewTransform.GetLocation();
+				Instance.Rotation		= NewTransform.GetRotation().Rotator();
+				Instance.DrawScale3D	= NewTransform.GetScale3D();
+			
+				// Rehash instance location
+				MeshInfo.InstanceHash->InsertInstance(Instance.Location, InstanceIdx);
+			}
+
+			for (auto It = MeshInfo.ComponentHash.CreateIterator(); It; ++It)
+			{
+				// We assume here that component we painted foliage on, was transformed as well
+				FFoliageComponentHashInfo& Info = It.Value();
+				Info.UpdateLocationFromActor(It.Key());
+			}
+		}
+	}
+}
+
+#endif // WITH_EDITOR
+
+struct FFoliageMeshInfo_Old
+{
+	TArray<FFoliageInstanceCluster_Deprecated> InstanceClusters;
+	TArray<FFoliageInstance> Instances;
+	UFoliageType_InstancedStaticMesh* Settings; // Type remapped via +ActiveClassRedirects
+};
+FArchive& operator<<(FArchive& Ar, FFoliageMeshInfo_Old& MeshInfo)
+{
+	Ar << MeshInfo.InstanceClusters;
+	Ar << MeshInfo.Instances;
+	Ar << MeshInfo.Settings;
+
+	return Ar;
+}
+
+void AInstancedFoliageActor::Serialize(FArchive& Ar)
+{
+	Super::Serialize(Ar);
+
+	Ar.UsingCustomVersion(FFoliageCustomVersion::GUID);
+
+	if (Ar.UE4Ver() < VER_UE4_FOLIAGE_SETTINGS_TYPE)
+	{
+		TMap<UStaticMesh*, FFoliageMeshInfo_Old> OldFoliageMeshes;
+		Ar << OldFoliageMeshes;
+		for (auto& OldMeshInfo : OldFoliageMeshes)
+		{
+			FFoliageMeshInfo NewMeshInfo;
+
+#if WITH_EDITORONLY_DATA
+			NewMeshInfo.Instances = MoveTemp(OldMeshInfo.Value.Instances);
+#endif
+
+			UFoliageType_InstancedStaticMesh* FoliageType = OldMeshInfo.Value.Settings;
+			if (FoliageType == nullptr)
+			{
+				// If the Settings object was null, eg the user forgot to save their settings asset, create a new one.
+				FoliageType = ConstructObject<UFoliageType_InstancedStaticMesh>(UFoliageType_InstancedStaticMesh::StaticClass(), this);
+			}
+
+			if (FoliageType->Mesh == nullptr)
+			{
+				FoliageType->Modify();
+				FoliageType->Mesh = OldMeshInfo.Key;
+			}
+			else if (FoliageType->Mesh != OldMeshInfo.Key)
+			{
+				// If mesh doesn't match (two meshes sharing the same settings object?) then we need to duplicate as that is no longer supported
+				FoliageType = (UFoliageType_InstancedStaticMesh*)StaticDuplicateObject(FoliageType, this, nullptr, RF_AllFlags & ~(RF_Standalone | RF_Public));
+				FoliageType->Mesh = OldMeshInfo.Key;
+			}
+			FoliageMeshes.Add(FoliageType, TUniqueObj<FFoliageMeshInfo>(MoveTemp(NewMeshInfo)));
+		}
+	}
+	else
+	{
+		Ar << FoliageMeshes;
+	}
+
+	// Clean up any old cluster components and convert to hierarchical instanced foliage.
+	if (Ar.CustomVer(FFoliageCustomVersion::GUID) < FFoliageCustomVersion::FoliageUsingHierarchicalISMC)
+	{
+		TInlineComponentArray<UInstancedStaticMeshComponent*> ClusterComponents;
+		GetComponents(ClusterComponents);
+		for (UInstancedStaticMeshComponent* Component : ClusterComponents)
+		{
+			Component->bAutoRegister = false;
+		}
+	}
+}
+
+void AInstancedFoliageActor::PostLoad()
+{
+	Super::PostLoad();
+
+#if WITH_EDITORONLY_DATA
+	if (GetLinkerUE4Version() < VER_UE4_DISALLOW_FOLIAGE_ON_BLUEPRINTS)
+	{
+		for (auto& MeshPair : FoliageMeshes)
+		{
+			for (FFoliageInstance& Instance : MeshPair.Value->Instances)
+			{
+				// Clear out the Base for any instances based on blueprint-created components,
+				// as those components will be destroyed when the construction scripts are
+				// re-run, leaving dangling references and causing crashes (woo!)
+				if (Instance.Base && Instance.Base->bCreatedByConstructionScript)
+				{
+					Instance.Base = NULL;
+				}
+			}
+		}
+	}
+#endif
+
+	if (RootComponent)
+	{
+		RootComponent->SetMobility(EComponentMobility::Static);	//instanced foliage is always static
+	}
+
+#if WITH_EDITOR
+	if (GIsEditor)
+	{
+		{
+			bool bContainsNull = FoliageMeshes.Remove(nullptr) > 0;
+			if (bContainsNull)
+			{
+				FMessageLog("MapCheck").Warning()
+					->AddToken(FUObjectToken::Create(this))
+					->AddToken(FTextToken::Create(LOCTEXT("MapCheck_Message_FoliageMissingStaticMesh", "Foliage instances for a missing static mesh have been removed.")))
+					->AddToken(FMapErrorToken::Create(FMapErrors::FoliageMissingStaticMesh));
+				while (bContainsNull)
+				{
+					bContainsNull = FoliageMeshes.Remove(nullptr) > 0;
+				}
+			}
+		}
+		for (auto& MeshPair : FoliageMeshes)
+		{
+			// Find the per-mesh info matching the mesh.
+			FFoliageMeshInfo& MeshInfo = *MeshPair.Value;
+			UFoliageType* FoliageType = MeshPair.Key;
+
+			if (MeshInfo.Instances.Num() && MeshInfo.Component == nullptr)
+			{
+				const UStaticMesh* StaticMesh = FoliageType->GetStaticMesh();
+				FFormatNamedArguments Arguments;
+				if (StaticMesh)
+				{
+					Arguments.Add(TEXT("MeshName"), FText::FromString(StaticMesh->GetName()));
+				}
+				else
+				{
+					Arguments.Add(TEXT("MeshName"), FText::FromString(TEXT("None")));
+				}
+
+				FMessageLog("MapCheck").Warning()
+					->AddToken(FUObjectToken::Create(this))
+					->AddToken(FTextToken::Create(FText::Format(LOCTEXT("MapCheck_Message_FoliageMissingComponent", "Foliage in this map is missing a component for static mesh {MeshName}. This has been repaired."),Arguments)))
+					->AddToken(FMapErrorToken::Create(FMapErrors::FoliageMissingClusterComponent));
+
+				MeshInfo.ReallocateClusters(this, MeshPair.Key);
+			}
+
+			// Update the hash.
+			for (int32 InstanceIdx = 0; InstanceIdx < MeshInfo.Instances.Num(); InstanceIdx++)
+			{
+				FFoliageInstance& Instance = MeshInfo.Instances[InstanceIdx];
+				{
+					// Add valid instances to the hash.
+					MeshInfo.InstanceHash->InsertInstance(Instance.Location, InstanceIdx);
+					FFoliageComponentHashInfo* ComponentHashInfo = MeshInfo.ComponentHash.Find(Instance.Base);
+					if (ComponentHashInfo == nullptr)
+					{
+						ComponentHashInfo = &MeshInfo.ComponentHash.Add(Instance.Base, FFoliageComponentHashInfo(Instance.Base));
+					}
+					ComponentHashInfo->Instances.Add(InstanceIdx);
+				}
+			}
+
+			// Convert to Heirarchical foliage
+			if (GetLinkerCustomVersion(FFoliageCustomVersion::GUID) < FFoliageCustomVersion::FoliageUsingHierarchicalISMC)
+			{
+				MeshInfo.ReallocateClusters(this, MeshPair.Key);
+			}
+
+			if (GetLinkerCustomVersion(FFoliageCustomVersion::GUID) < FFoliageCustomVersion::HierarchicalISMCNonTransactional)
+			{
+				if (MeshInfo.Component)
+				{
+					MeshInfo.Component->ClearFlags(RF_Transactional);
+				}
+			}
+		}
+	}
+#endif
+}
+
+#if WITH_EDITOR
+
+void AInstancedFoliageActor::OnLevelActorMoved(AActor* InActor)
+{
+	if (!GetWorld()->IsPlayInEditor())
+	{
+		if(AInstancedFoliageActor* IFA = AInstancedFoliageActor::GetInstancedFoliageActorForLevel(InActor->GetLevel()))
+		{
+			TInlineComponentArray<UActorComponent*> Components;
+			GetComponents(Components);
+
+			for (int32 Idx = 0; Idx < Components.Num(); ++Idx)
+			{
+				IFA->MoveInstancesForMovedComponent(Components[Idx]);
+			}
+		}
+	}
+}
+
+void AInstancedFoliageActor::SubscribeToUpdateEvents()
+{
+	if (!GetWorld()->IsPlayInEditor())
+	{
+		OnLevelActorMovedDelegateHandle   = GEngine->OnActorMoved().AddUObject(this, &AInstancedFoliageActor::OnLevelActorMoved);
+	}
+}
+#endif
+
+//
+// Serialize all our UObjects for RTGC 
+//
+void AInstancedFoliageActor::AddReferencedObjects(UObject* InThis, FReferenceCollector& Collector)
+{
+	AInstancedFoliageActor* This = CastChecked<AInstancedFoliageActor>(InThis);
+	if (This->SelectedMesh)
+	{
+		Collector.AddReferencedObject(This->SelectedMesh, This);
+	}
+
+	for (auto& MeshPair : This->FoliageMeshes)
+	{
+		Collector.AddReferencedObject(MeshPair.Key, This);
+		FFoliageMeshInfo& MeshInfo = *MeshPair.Value;
+
+#if WITH_EDITORONLY_DATA
+		for (FFoliageInstance& Instance : MeshInfo.Instances)
+		{
+			if (Instance.Base != nullptr)
+			{
+				Collector.AddReferencedObject(Instance.Base, This);
+			}
+		}
+#endif
+
+		if (MeshInfo.Component)
+		{
+			Collector.AddReferencedObject(MeshInfo.Component, This);
+		}
+	}
+	Super::AddReferencedObjects(This, Collector);
+}
+
+
+
+void AInstancedFoliageActor::ApplyWorldOffset(const FVector& InOffset, bool bWorldShift)
+{
+	Super::ApplyWorldOffset(InOffset, bWorldShift);
+
+	if (GIsEditor)
+	{
+		for (auto& MeshPair : FoliageMeshes)
+		{
+			FFoliageMeshInfo& MeshInfo = *MeshPair.Value;
+
+#if WITH_EDITORONLY_DATA
+			MeshInfo.InstanceHash->Empty();
+			for (int32 InstanceIdx = 0; InstanceIdx < MeshInfo.Instances.Num(); InstanceIdx++)
+			{
+				FFoliageInstance& Instance = MeshInfo.Instances[InstanceIdx];
+				Instance.Location += InOffset;
+				// Rehash instance location
+				MeshInfo.InstanceHash->InsertInstance(Instance.Location, InstanceIdx);
+			}
+
+			for (auto It = MeshInfo.ComponentHash.CreateIterator(); It; ++It)
+			{
+				// We assume here that component we painted foliage on will be shifted by same value
+				FFoliageComponentHashInfo& Info = It.Value();
+				Info.CachedLocation += InOffset;
+			}
+#endif
+		}
+	}
+}
+
+#undef LOCTEXT_NAMESPACE
diff --git a/Engine/Source/Runtime/Foliage/Private/InteractiveFoliageActor.cpp b/Engine/Source/Runtime/Foliage/Private/InteractiveFoliageActor.cpp
new file mode 100644
--- /dev/null
+++ b/Engine/Source/Runtime/Foliage/Private/InteractiveFoliageActor.cpp
@@ -0,0 +1,37 @@
+// Copyright 1998-2015 Epic Games, Inc. All Rights Reserved.
+
+#include "FoliagePrivate.h"
+#include "InteractiveFoliageActor.h"
+#include "InteractiveFoliageComponent.h"
+#include "Components/CapsuleComponent.h"
+
+AInteractiveFoliageActor::AInteractiveFoliageActor(const FObjectInitializer& ObjectInitializer)
+	: Super(ObjectInitializer
+		.SetDefaultSubobjectClass<UInteractiveFoliageComponent>("StaticMeshComponent0"))
+{
+
+	UInteractiveFoliageComponent* FoliageMeshComponent = CastChecked<UInteractiveFoliageComponent>(GetStaticMeshComponent());
+	FoliageMeshComponent->SetCollisionProfileName(UCollisionProfile::NoCollision_ProfileName);
+	FoliageMeshComponent->Mobility = EComponentMobility::Static;
+
+	CapsuleComponent = CreateDefaultSubobject<UCapsuleComponent>(TEXT("CollisionCylinder"));
+	CapsuleComponent->InitCapsuleSize(60.0f, 200.0f);
+	static FName CollisionProfileName(TEXT("OverlapAllDynamic"));
+	CapsuleComponent->SetCollisionProfileName(CollisionProfileName);
+	CapsuleComponent->Mobility = EComponentMobility::Static;
+
+	RootComponent = CapsuleComponent;
+
+	PrimaryActorTick.bCanEverTick = true;
+	bCanBeDamaged = true;
+	bCollideWhenPlacing = true;
+	FoliageDamageImpulseScale = 20.0f;
+	FoliageTouchImpulseScale = 10.0f;
+	FoliageStiffness = 10.0f;
+	FoliageStiffnessQuadratic = 0.3f;
+	FoliageDamping = 2.0f;
+	MaxDamageImpulse = 100000.0f;
+	MaxTouchImpulse = 1000.0f;
+	MaxForce = 100000.0f;
+	Mass = 1.0f;
+}
diff --git a/Engine/Source/Runtime/Foliage/Private/InteractiveFoliageComponent.h b/Engine/Source/Runtime/Foliage/Private/InteractiveFoliageComponent.h
new file mode 100644
--- /dev/null
+++ b/Engine/Source/Runtime/Foliage/Private/InteractiveFoliageComponent.h
@@ -0,0 +1,26 @@
+// Copyright 1998-2015 Epic Games, Inc. All Rights Reserved.
+
+
+#pragma once
+#include "InteractiveFoliageComponent.generated.h"
+
+UCLASS(hidecategories=Object)
+class UInteractiveFoliageComponent : public UStaticMeshComponent
+{
+	GENERATED_UCLASS_BODY()
+
+
+public:
+	class FInteractiveFoliageSceneProxy* FoliageSceneProxy;
+
+	// Begin UPrimitiveComponent Interface
+	virtual FPrimitiveSceneProxy* CreateSceneProxy() override;
+	// End UPrimitiveComponent Interface
+
+	// Begin UActorComponent Interface
+	virtual void DestroyRenderState_Concurrent() override;
+	// End UActorComponent Interface
+
+	friend class AInteractiveFoliageActor;
+};
+
diff --git a/Engine/Source/Runtime/Foliage/Public/FoliageType.h b/Engine/Source/Runtime/Foliage/Public/FoliageType.h
new file mode 100644
--- /dev/null
+++ b/Engine/Source/Runtime/Foliage/Public/FoliageType.h
@@ -0,0 +1,309 @@
+// Copyright 1998-2015 Epic Games, Inc. All Rights Reserved.
+
+#pragma once
+#include "Curves/CurveFloat.h"
+#include "FoliageType.generated.h"
+
+UENUM()
+enum FoliageVertexColorMask
+{
+	FOLIAGEVERTEXCOLORMASK_Disabled,
+	FOLIAGEVERTEXCOLORMASK_Red,
+	FOLIAGEVERTEXCOLORMASK_Green,
+	FOLIAGEVERTEXCOLORMASK_Blue,
+	FOLIAGEVERTEXCOLORMASK_Alpha,
+};
+
+
+UCLASS(hidecategories = Object, editinlinenew, MinimalAPI, BlueprintType, Blueprintable)
+class UFoliageType : public UObject
+{
+	GENERATED_UCLASS_BODY()
+
+	UPROPERTY(EditAnywhere, Category=Painting)
+	float Density;
+
+	UPROPERTY(EditAnywhere, Category=Painting)
+	float Radius;
+
+	UPROPERTY(EditAnywhere, Category=Painting)
+	float ScaleMinX;
+
+	UPROPERTY(EditAnywhere, Category=Painting)
+	float ScaleMinY;
+
+	UPROPERTY(EditAnywhere, Category=Painting)
+	float ScaleMinZ;
+
+	UPROPERTY(EditAnywhere, Category=Painting)
+	float ScaleMaxX;
+
+	UPROPERTY(EditAnywhere, Category=Painting)
+	float ScaleMaxY;
+
+	UPROPERTY(EditAnywhere, Category=Painting)
+	float ScaleMaxZ;
+
+	UPROPERTY(EditAnywhere, Category=Painting)
+	uint32 LockScaleX:1;
+
+	UPROPERTY(EditAnywhere, Category=Painting)
+	uint32 LockScaleY:1;
+
+	UPROPERTY(EditAnywhere, Category=Painting)
+	uint32 LockScaleZ:1;
+
+	UPROPERTY(EditAnywhere, Category=General, meta=(Subcategory = "Placement"))
+	float AlignMaxAngle;
+
+	UPROPERTY(EditAnywhere, Category = General, meta = (Subcategory = "Placement"))
+	float RandomPitchAngle;
+
+	UPROPERTY(EditAnywhere, Category = General, meta = (Subcategory = "Placement"))
+	float GroundSlope;
+
+	UPROPERTY(EditAnywhere, Category = General, meta = (Subcategory = "Placement"))
+	float HeightMin;
+
+	UPROPERTY(EditAnywhere, Category = General, meta = (Subcategory = "Placement"))
+	float HeightMax;
+
+	UPROPERTY(EditAnywhere, Category = General, meta = (Subcategory = "Placement"))
+	FName LandscapeLayer;
+
+	UPROPERTY(EditAnywhere, Category = General, meta = (Subcategory = "Placement"))
+	uint32 AlignToNormal:1;
+
+	UPROPERTY(EditAnywhere, Category = General, meta = (Subcategory = "Placement"))
+	uint32 RandomYaw:1;
+
+	UPROPERTY(EditAnywhere, Category=Painting)
+	uint32 UniformScale:1;
+
+	UPROPERTY(EditAnywhere, Category=Painting)
+	float ZOffsetMin;
+
+	UPROPERTY(EditAnywhere, Category=Painting)
+	float ZOffsetMax;
+
+	UPROPERTY(EditAnywhere, Category=Painting)
+	uint32 CollisionWithWorld:1;
+
+	UPROPERTY(EditAnywhere, Category=Painting)
+	FVector CollisionScale;
+
+	UPROPERTY()
+	FBoxSphereBounds MeshBounds;
+
+	// X, Y is origin position and Z is radius...
+	UPROPERTY()
+	FVector LowBoundOriginRadius;
+
+	UPROPERTY(EditAnywhere, Category=Painting)
+	TEnumAsByte<enum FoliageVertexColorMask> VertexColorMask;
+
+	UPROPERTY(EditAnywhere, Category=Painting)
+	float VertexColorMaskThreshold;
+
+	UPROPERTY(EditAnywhere, Category=Painting)
+	uint32 VertexColorMaskInvert:1;
+
+	UPROPERTY()
+	float ReapplyDensityAmount;
+
+	UPROPERTY()
+	uint32 ReapplyDensity:1;
+
+	UPROPERTY()
+	uint32 ReapplyRadius:1;
+
+	UPROPERTY()
+	uint32 ReapplyAlignToNormal:1;
+
+	UPROPERTY()
+	uint32 ReapplyRandomYaw:1;
+
+	UPROPERTY()
+	uint32 ReapplyScaleX:1;
+
+	UPROPERTY()
+	uint32 ReapplyScaleY:1;
+
+	UPROPERTY()
+	uint32 ReapplyScaleZ:1;
+
+	UPROPERTY()
+	uint32 ReapplyRandomPitchAngle:1;
+
+	UPROPERTY()
+	uint32 ReapplyGroundSlope:1;
+
+	UPROPERTY()
+	uint32 ReapplyHeight:1;
+
+	UPROPERTY()
+	uint32 ReapplyLandscapeLayer:1;
+
+	UPROPERTY()
+	uint32 ReapplyZOffset:1;
+
+	UPROPERTY()
+	uint32 ReapplyCollisionWithWorld:1;
+
+	UPROPERTY()
+	uint32 ReapplyVertexColorMask:1;
+
+	/* The distance where instances will begin to fade out if using a PerInstanceFadeAmount material node. 0 disables. */
+	UPROPERTY(EditAnywhere, Category = General, meta = (Subcategory = "Culling"))
+	int32 StartCullDistance;
+
+	/**
+	 * The distance where instances will have completely faded out when using a PerInstanceFadeAmount material node. 0 disables. 
+	 * When the entire cluster is beyond this distance, the cluster is completely culled and not rendered at all.
+	 */
+	UPROPERTY(EditAnywhere, Category = General, meta = (Subcategory = "Culling"))
+	int32 EndCullDistance;
+
+	UPROPERTY()
+	int32 DisplayOrder;
+
+	UPROPERTY()
+	uint32 IsSelected:1;
+
+	UPROPERTY()
+	uint32 ShowNothing:1;
+
+	UPROPERTY()
+	uint32 ShowPaintSettings:1;
+
+	UPROPERTY()
+	uint32 ShowInstanceSettings:1;
+
+	/** Controls whether the foliage should take part in static lighting/shadowing. If false, mesh will not receive or cast static lighting or shadows, regardless of other settings. */
+	UPROPERTY(EditAnywhere, BlueprintReadOnly, Category = General, meta = (Subcategory = "Lighting"))
+	uint32 bEnableStaticLighting : 1;
+
+	/** Controls whether the foliage should cast a shadow or not. */
+	UPROPERTY(EditAnywhere, BlueprintReadOnly, Category = General, meta = (Subcategory = "Lighting"))
+	uint32 CastShadow:1;
+
+	/** Controls whether the foliage should inject light into the Light Propagation Volume.  This flag is only used if CastShadow is true. */
+	UPROPERTY(EditAnywhere, BlueprintReadOnly, Category = General, meta = (Subcategory = "Lighting"))
+	uint32 bAffectDynamicIndirectLighting:1;
+
+	/** Controls whether the primitive should affect dynamic distance field lighting methods.  This flag is only used if CastShadow is true. */
+	UPROPERTY(EditAnywhere, BlueprintReadOnly, Category = General, meta = (Subcategory = "Lighting", EditCondition = "CastShadow"))
+	uint32 bAffectDistanceFieldLighting:1;
+
+	/** Controls whether the foliage should cast shadows in the case of non precomputed shadowing.  This flag is only used if CastShadow is true. */
+	UPROPERTY(EditAnywhere, BlueprintReadOnly, Category = General, meta = (Subcategory = "Lighting"))
+	uint32 bCastDynamicShadow:1;
+
+	/** Whether the foliage should cast a static shadow from shadow casting lights.  This flag is only used if CastShadow is true. */
+	UPROPERTY(EditAnywhere, BlueprintReadOnly, Category = General, meta = (Subcategory = "Lighting"))
+	uint32 bCastStaticShadow:1;
+
+	/** Whether this foliage should cast dynamic shadows as if it were a two sided material. */
+	UPROPERTY(EditAnywhere, AdvancedDisplay, BlueprintReadOnly, Category = General, meta = (Subcategory = "Lighting"))
+	uint32 bCastShadowAsTwoSided:1;
+
+	/** Whether the foliage receives decals. */
+	UPROPERTY(EditAnywhere, AdvancedDisplay, BlueprintReadOnly, Category = General, meta = (Subcategory = "Rendering"))
+	uint32 bReceivesDecals : 1;
+	
+	/** Whether to override the lightmap resolution defined in the static mesh. */
+	UPROPERTY(BlueprintReadOnly, Category = General, meta = (Subcategory = "Lighting"))
+	uint32 bOverrideLightMapRes:1;
+
+	/** Overrides the lightmap resolution defined in the static mesh */
+	UPROPERTY(EditAnywhere, BlueprintReadOnly, Category = General, meta = (Subcategory = "Lighting", DisplayName = "Light Map Resolution", EditCondition = "bOverrideLightMapRes"))
+	int32 OverriddenLightMapRes;
+
+	/** Custom collision for foliage */
+	UPROPERTY(EditAnywhere, Category = Painting, meta = (Subcategory = "Collision", HideObjectType = true))
+	struct FBodyInstance BodyInstance;
+
+	/* Gets the mesh associated with this FoliageType */
+	virtual UStaticMesh* GetStaticMesh() PURE_VIRTUAL(UFoliageType::GetStaticMesh, return nullptr; );
+
+#if WITH_EDITOR
+	virtual void PostEditChangeProperty(struct FPropertyChangedEvent& PropertyChangedEvent) override;
+#endif
+
+	/** Ecosystem specific parameters */
+
+	/** Specifies the number of seeds to populate along 10 meters. The number is implicitly squared to cover a 10m x 10m area*/
+	UPROPERTY(Category = Ecosystem, EditAnywhere, meta = (Subcategory = "Clustering", ClampMin = "0.0", UIMin = "0.0"))
+	float InitialSeedDensity;
+
+	/** The average distance between the spreading instance and its seeds. For example, a tree with an AverageSpreadDistance 10 will ensure the average distance between the tree and its seeds is 10cm */
+	UPROPERTY(Category = Ecosystem, EditAnywhere, meta = (Subcategory = "Clustering", ClampMin = "0.0", UIMin = "0.0"))
+	float AverageSpreadDistance;
+
+	/** Specifies how much seed distance varies from the average. For example, a tree with an AverageSpreadDistance 10 and a SpreadVariance 1 will produce seeds with an average distance of 10cm plus or minus 1cm */
+	UPROPERTY(Category = Ecosystem, EditAnywhere, meta = (Subcategory = "Clustering", ClampMin = "0.0", UIMin = "0.0"))
+	float SpreadVariance;
+
+	/** The CollisionRadius determines when two instances overlap. When two instances overlap a winner will be picked based on rules and priority. */
+	UPROPERTY(Category = Ecosystem, EditAnywhere, meta = (Subcategory = "Collision", ClampMin = "0.0", UIMin = "0.0"))
+	float CollisionRadius;
+
+	/** The ShadeRadius determines when two instances overlap. If an instance can grow in the shade this radius is ignored.*/
+	UPROPERTY(Category = Ecosystem, EditAnywhere, meta = (Subcategory = "Collision", ClampMin = "0.0", UIMin = "0.0"))
+	float ShadeRadius;
+
+	/** The number of times we age the species and spread its seeds. */
+	UPROPERTY(Category = Ecosystem, EditAnywhere, meta = (Subcategory = "Clustering", ClampMin = "0", UIMin = "0"))
+	int32 NumSteps;
+
+	/** The number of seeds an instance will spread in a single step of the simulation. */
+	UPROPERTY(Category = Ecosystem, EditAnywhere, meta = (Subcategory = "Clustering", ClampMin = "0", UIMin = "0"))
+	int32 SeedsPerStep;
+
+	/** Whether the species can grow in shade. If this is true shade radius is ignored during overlap tests*/
+	UPROPERTY(Category = Ecosystem, EditAnywhere, meta = (Subcategory = "Growth"))
+	bool bGrowsInShade;
+
+	/** The minimum scale that an instance will be. Corresponds to age = 0*/
+	UPROPERTY(Category = Ecosystem, EditAnywhere, meta = (Subcategory = "Growth", ClampMin = "0.0", UIMin = "0.0"))
+	float MinScale;
+
+	/** The maximum scale that a seed will grow to. Corresponds to age = 1*/
+	UPROPERTY(Category = Ecosystem, EditAnywhere, meta = (Subcategory = "Clustering", ClampMin = "0.0", UIMin = "0.0"))
+	float MaxScale;
+
+	/** Specifies the oldest a new seed can be. The new seed will have an age randomly distributed in [0,InitMaxAge] */
+	UPROPERTY(Category = Ecosystem, EditAnywhere, meta = (Subcategory = "Growth", ClampMin = "0.0", UIMin = "0.0"))
+	float InitialMaxAge;
+
+	/** Specifies the oldest a seed can be. After reaching this age the instance will still spread seeds, but will not get any older*/
+	UPROPERTY(Category = Ecosystem, EditAnywhere, meta = (Subcategory = "Growth", ClampMin = "0.0", UIMin = "0.0"))
+	float MaxAge;
+
+	/** When two instances overlap we must determine which instance to remove. The instance with a lower OverlapPriority will be removed. In the case where OverlapPriority is the same regular simulation rules apply.*/
+	UPROPERTY(Category = Ecosystem, EditAnywhere, meta = (Subcategory = "Growth"))
+	float OverlapPriority;
+
+	/** The set of meshes to be used for representing an instance placed in the world.*/
+	UPROPERTY(Category = Ecosystem, EditAnywhere, meta = (Subcategory = "Rendering"))
+	TArray<UStaticMesh*> StaticMeshes;
+
+
+	UPROPERTY()
+	int32 ChangeCount;
+
+	
+	FOLIAGE_API float GetSeedDensitySquared() const { return InitialSeedDensity * InitialSeedDensity; }
+	FOLIAGE_API float GetMaxRadius() const;
+	FOLIAGE_API float GetScaleForAge(const float Age) const;
+	FOLIAGE_API float GetInitAge(FRandomStream& RandomStream) const;
+	FOLIAGE_API float GetNextAge(const float CurrentAge, const int32 NumSteps) const;
+
+private:
+
+	/** The curve used to interpolate the instance scale.*/
+	UPROPERTY(Category = Ecosystem, EditAnywhere, meta = (Subcategory = "Rendering"))
+	FRuntimeFloatCurve ScaleCurve;
+};
+
+
diff --git a/Engine/Source/Runtime/Foliage/Public/FoliageType_InstancedStaticMesh.h b/Engine/Source/Runtime/Foliage/Public/FoliageType_InstancedStaticMesh.h
new file mode 100644
--- /dev/null
+++ b/Engine/Source/Runtime/Foliage/Public/FoliageType_InstancedStaticMesh.h
@@ -0,0 +1,20 @@
+// Copyright 1998-2015 Epic Games, Inc. All Rights Reserved.
+
+#pragma once
+
+#include "FoliageType.h"
+#include "FoliageType_InstancedStaticMesh.generated.h"
+
+UCLASS(hidecategories=Object, editinlinenew, MinimalAPI)
+class UFoliageType_InstancedStaticMesh : public UFoliageType
+{
+	GENERATED_UCLASS_BODY()
+
+	UPROPERTY(EditAnywhere, Category=FoliageType)
+	UStaticMesh* Mesh;
+
+	virtual UStaticMesh* GetStaticMesh() override
+	{
+		return Mesh;
+	}
+};
diff --git a/Engine/Source/Runtime/Foliage/Public/InstancedFoliage.h b/Engine/Source/Runtime/Foliage/Public/InstancedFoliage.h
new file mode 100644
--- /dev/null
+++ b/Engine/Source/Runtime/Foliage/Public/InstancedFoliage.h
@@ -0,0 +1,343 @@
+// Copyright 1998-2015 Epic Games, Inc. All Rights Reserved.
+
+/*=============================================================================
+	InstancedFoliage.h: Instanced foliage type definitions.
+  =============================================================================*/
+
+#pragma once
+
+//
+// Forward declarations.
+//
+class UInstancedStaticMeshComponent;
+class UHierarchicalInstancedStaticMeshComponent;
+class AInstancedFoliageActor;
+class UFoliageType;
+struct FFoliageInstanceHash;
+
+/**
+ * Flags stored with each instance
+ */
+enum EFoliageInstanceFlags
+{
+	FOLIAGE_AlignToNormal	= 0x00000001,
+	FOLIAGE_NoRandomYaw		= 0x00000002,
+	FOLIAGE_Readjusted		= 0x00000004,
+	FOLIAGE_InstanceDeleted	= 0x00000008,	// Used only for migration from pre-HierarchicalISM foliage.
+};
+
+/**
+ *	FFoliageInstancePlacementInfo - placement info an individual instance
+ */
+struct FFoliageInstancePlacementInfo
+{
+	FVector Location;
+	FRotator Rotation;
+	FRotator PreAlignRotation;
+	FVector DrawScale3D;
+	float ZOffset;
+	uint32 Flags;
+
+	FFoliageInstancePlacementInfo()
+		: Location(0.f, 0.f, 0.f)
+		, Rotation(0, 0, 0)
+		, PreAlignRotation(0, 0, 0)
+		, DrawScale3D(1.f, 1.f, 1.f)
+		, ZOffset(0.f)
+		, Flags(0)
+	{}
+};
+
+/**
+ *	FFoliageInstance - editor info an individual instance
+ */
+struct FFoliageInstance : public FFoliageInstancePlacementInfo
+{
+	UActorComponent* Base;
+
+	UActorComponent* Spawner;
+
+	FFoliageInstance()
+	: Base(NULL)
+	, Spawner(NULL)
+	{}
+
+
+	friend FArchive& operator<<(FArchive& Ar, FFoliageInstance& Instance);
+
+	FTransform GetInstanceWorldTransform() const
+	{
+		return FTransform(Rotation, Location, DrawScale3D);
+	}
+
+	void AlignToNormal(const FVector& InNormal, float AlignMaxAngle = 0.f)
+	{
+		Flags |= FOLIAGE_AlignToNormal;
+
+		FRotator AlignRotation = InNormal.Rotation();
+		// Static meshes are authored along the vertical axis rather than the X axis, so we add 90 degrees to the static mesh's Pitch.
+		AlignRotation.Pitch -= 90.f;
+		// Clamp its value inside +/- one rotation
+		AlignRotation.Pitch = FRotator::NormalizeAxis(AlignRotation.Pitch);
+
+		// limit the maximum pitch angle if it's > 0.
+		if (AlignMaxAngle > 0.f)
+		{
+			int32 MaxPitch = AlignMaxAngle;
+			if (AlignRotation.Pitch > MaxPitch)
+			{
+				AlignRotation.Pitch = MaxPitch;
+			}
+			else if (AlignRotation.Pitch < -MaxPitch)
+			{
+				AlignRotation.Pitch = -MaxPitch;
+			}
+		}
+
+		PreAlignRotation = Rotation;
+		Rotation = FRotator(FQuat(AlignRotation) * FQuat(Rotation));
+	}
+};
+
+/**
+ * FFoliageComponentHashInfo
+ * Cached instance list and component location info stored in the ComponentHash.
+ * Used for moving quick updates after operations on components with foliage painted on them.
+ */
+struct FFoliageComponentHashInfo
+{
+	// tors
+	FFoliageComponentHashInfo()
+		: CachedLocation(0, 0, 0)
+		, CachedRotation(0, 0, 0)
+		, CachedDrawScale(1, 1, 1)
+	{}
+
+	FFoliageComponentHashInfo(UActorComponent* InComponent)
+		: CachedLocation(0, 0, 0)
+		, CachedRotation(0, 0, 0)
+		, CachedDrawScale(1, 1, 1)
+	{
+		UpdateLocationFromActor(InComponent);
+	}
+
+	// Cache the location and rotation from the actor
+	void UpdateLocationFromActor(UActorComponent* InComponent)
+	{
+		if (InComponent)
+		{
+			AActor* Owner = Cast<AActor>(InComponent->GetOuter());
+			if (Owner)
+			{
+				const USceneComponent* RootComponent = Owner->GetRootComponent();
+				if (RootComponent)
+				{
+					CachedLocation = RootComponent->RelativeLocation;
+					CachedRotation = RootComponent->RelativeRotation;
+					CachedDrawScale = RootComponent->RelativeScale3D;
+				}
+			}
+		}
+	}
+
+	// serializer
+	friend FArchive& operator<<(FArchive& Ar, FFoliageComponentHashInfo& ComponentHashInfo);
+
+	FVector CachedLocation;
+	FRotator CachedRotation;
+	FVector CachedDrawScale;
+	TSet<int32> Instances;
+};
+
+/**
+ *	FFoliageMeshInfo - editor info for all matching foliage meshes
+ */
+struct FFoliageMeshInfo
+{
+	UHierarchicalInstancedStaticMeshComponent* Component;
+
+#if WITH_EDITORONLY_DATA
+	// Editor-only placed instances
+	TArray<FFoliageInstance> Instances;
+
+	// Transient, editor-only locality hash of instances
+	TUniquePtr<FFoliageInstanceHash> InstanceHash;
+
+	// Transient, editor-only set of instances per component
+	TMap<UActorComponent*, FFoliageComponentHashInfo> ComponentHash;
+
+	// Transient, editor-only list of selected instances.
+	TSet<int32> SelectedIndices;
+#endif
+
+	FOLIAGE_API FFoliageMeshInfo();
+
+	~FFoliageMeshInfo() // =default
+	{ }
+
+	FFoliageMeshInfo(FFoliageMeshInfo&& Other)
+		// even VC++2013 doesn't support "=default" on move constructors
+		: Component(Other.Component)
+#if WITH_EDITORONLY_DATA
+		, Instances(MoveTemp(Other.Instances))
+		, InstanceHash(MoveTemp(Other.InstanceHash))
+		, ComponentHash(MoveTemp(Other.ComponentHash))
+		, SelectedIndices(MoveTemp(Other.SelectedIndices))
+#endif
+	{ }
+
+	FFoliageMeshInfo& operator=(FFoliageMeshInfo&& Other)
+		// even VC++2013 doesn't support "=default" on move assignment
+	{
+		Component = Other.Component;
+#if WITH_EDITORONLY_DATA
+		Instances = MoveTemp(Other.Instances);
+		InstanceHash = MoveTemp(Other.InstanceHash);
+		ComponentHash = MoveTemp(Other.ComponentHash);
+		SelectedIndices = MoveTemp(Other.SelectedIndices);
+#endif
+
+		return *this;
+	}
+
+#if WITH_EDITOR
+	FOLIAGE_API void AddInstance(AInstancedFoliageActor* InIFA, UFoliageType* InSettings, const FFoliageInstance& InNewInstance);
+	FOLIAGE_API void RemoveInstances(AInstancedFoliageActor* InIFA, const TArray<int32>& InInstancesToRemove);
+	FOLIAGE_API void PreMoveInstances(AInstancedFoliageActor* InIFA, const TArray<int32>& InInstancesToMove);
+	FOLIAGE_API void PostMoveInstances(AInstancedFoliageActor* InIFA, const TArray<int32>& InInstancesMoved);
+	FOLIAGE_API void PostUpdateInstances(AInstancedFoliageActor* InIFA, const TArray<int32>& InInstancesUpdated, bool bReAddToHash = false);
+	FOLIAGE_API void DuplicateInstances(AInstancedFoliageActor* InIFA, UFoliageType* InSettings, const TArray<int32>& InInstancesToDuplicate);
+	FOLIAGE_API void GetInstancesInsideSphere(const FSphere& Sphere, TArray<int32>& OutInstances);
+	FOLIAGE_API bool CheckForOverlappingSphere(const FSphere& Sphere);
+	FOLIAGE_API bool CheckForOverlappingInstanceExcluding(int32 TestInstanceIdx, float Radius, TSet<int32>& ExcludeInstances);
+
+	// Destroy existing clusters and reassign all instances to new clusters
+	FOLIAGE_API void ReallocateClusters(AInstancedFoliageActor* InIFA, UFoliageType* InSettings);
+
+	FOLIAGE_API void ReapplyInstancesToComponent();
+
+	FOLIAGE_API void SelectInstances(AInstancedFoliageActor* InIFA, bool bSelect, TArray<int32>& Instances);
+
+	// Get the number of placed instances
+	FOLIAGE_API int32 GetInstanceCount() const;
+
+	// For debugging. Validate state after editing.
+	void CheckValid();
+#endif
+
+	friend FArchive& operator<<(FArchive& Ar, FFoliageMeshInfo& MeshInfo);
+
+#if PLATFORM_COMPILER_HAS_DEFAULTED_FUNCTIONS
+	FFoliageMeshInfo(const FFoliageMeshInfo& Other) = delete;
+	const FFoliageMeshInfo& operator=(const FFoliageMeshInfo& Other) = delete;
+#else
+private:
+	FFoliageMeshInfo(const FFoliageMeshInfo& Other);
+	const FFoliageMeshInfo& operator=(const FFoliageMeshInfo& Other);
+#endif
+};
+
+//
+// FFoliageInstanceHash
+//
+
+#define FOLIAGE_HASH_CELL_BITS 9	// 512x512 grid
+
+struct FFoliageInstanceHash
+{
+private:
+	const int32 HashCellBits;
+	TMultiMap<uint64, int32> CellMap;
+
+	uint64 MakeKey(int32 CellX, int32 CellY)
+	{
+		return ((uint64)(*(uint32*)(&CellX)) << 32) | (*(uint32*)(&CellY) & 0xffffffff);
+	}
+
+	uint64 MakeKey(const FVector& Location)
+	{
+		return  MakeKey(FMath::FloorToInt(Location.X) >> HashCellBits, FMath::FloorToInt(Location.Y) >> HashCellBits);
+	}
+
+	// Locality map
+	//TMap<uint64, TSet<int32>> CellMap;
+public:
+	FFoliageInstanceHash(int32 InHashCellBits = FOLIAGE_HASH_CELL_BITS)
+	:	HashCellBits(InHashCellBits)
+	{}
+
+	void InsertInstance(const FVector& InstanceLocation, int32 InstanceIndex)
+	{
+		uint64 Key = MakeKey(InstanceLocation);
+
+		CellMap.AddUnique(Key, InstanceIndex);
+	}
+
+	void RemoveInstance(const FVector& InstanceLocation, int32 InstanceIndex)
+	{
+		uint64 Key = MakeKey(InstanceLocation);
+		
+		int32 RemoveCount = CellMap.RemoveSingle(Key, InstanceIndex);
+		check(RemoveCount == 1);
+	}
+
+	void GetInstancesOverlappingBox(const FBox& InBox, TArray<int32>& OutInstanceIndices)
+	{
+		int32 MinX = FMath::FloorToInt(InBox.Min.X) >> HashCellBits;
+		int32 MinY = FMath::FloorToInt(InBox.Min.Y) >> HashCellBits;
+		int32 MaxX = FMath::FloorToInt(InBox.Max.X) >> HashCellBits;
+		int32 MaxY = FMath::FloorToInt(InBox.Max.Y) >> HashCellBits;
+
+		for (int32 y = MinY; y <= MaxY; y++)
+		{
+			for (int32 x = MinX; x <= MaxX; x++)
+			{
+				uint64 Key = MakeKey(x, y);
+				CellMap.MultiFind(Key, OutInstanceIndices);
+			}
+		}
+	}
+
+	TArray<int32> GetInstancesOverlappingBox(const FBox& InBox)
+	{
+		TArray<int32> Result;
+		GetInstancesOverlappingBox(InBox, Result);
+		return Result;
+	}
+
+#if UE_BUILD_DEBUG
+	void CheckInstanceCount(int32 InCount)
+	{
+		check(CellMap.Num() == InCount);
+	}
+#endif
+
+	void Empty()
+	{
+		CellMap.Empty();
+	}
+
+	friend FArchive& operator<<(FArchive& Ar, FFoliageInstanceHash& Hash)
+	{
+		if (Ar.UE4Ver() < VER_UE4_FOLIAGE_SETTINGS_TYPE)
+		{
+			Hash.CellMap.Reset();
+
+			TMap<uint64, TSet<int32>> OldCellMap;
+			Ar << OldCellMap;
+			for (auto& CellPair : OldCellMap)
+			{
+				for (int32 Idx : CellPair.Value)
+				{
+					Hash.CellMap.AddUnique(CellPair.Key, Idx);
+				}
+			}
+		}
+		else
+		{
+			Ar << Hash.CellMap;
+		}
+
+		return Ar;
+	}
+};
diff --git a/Engine/Source/Runtime/Foliage/Public/InstancedFoliageActor.h b/Engine/Source/Runtime/Foliage/Public/InstancedFoliageActor.h
new file mode 100644
--- /dev/null
+++ b/Engine/Source/Runtime/Foliage/Public/InstancedFoliageActor.h
@@ -0,0 +1,130 @@
+// Copyright 1998-2015 Epic Games, Inc. All Rights Reserved.
+
+
+#pragma once
+
+#include "GameFramework/Actor.h"
+#include "Templates/UniqueObj.h"
+#include "Components/InstancedStaticMeshComponent.h"
+
+#include "InstancedFoliageActor.generated.h"
+
+// Forward declarations
+class UFoliageType;
+struct FFoliageInstancePlacementInfo;
+struct FFoliageMeshInfo;
+
+UCLASS(notplaceable, hidecategories = Object, MinimalAPI, NotBlueprintable)
+class AInstancedFoliageActor : public AActor
+{
+	GENERATED_UCLASS_BODY()
+
+	/** The static mesh type that will be used to show the widget */
+	UPROPERTY(transient)
+	UFoliageType* SelectedMesh;
+
+public:
+	TMap<UFoliageType*, TUniqueObj<FFoliageMeshInfo>> FoliageMeshes;
+
+public:
+	// Begin UObject interface.
+	virtual void Serialize(FArchive& Ar) override;
+	virtual void PostLoad() override;
+	static void AddReferencedObjects(UObject* InThis, FReferenceCollector& Collector);
+	// End UObject interface. 
+
+	// Begin AActor interface.
+	// we don't want to have our components automatically destroyed by the Blueprint code
+	virtual void RerunConstructionScripts() override {}
+	// Add world origin offset
+	virtual void ApplyWorldOffset(const FVector& InOffset, bool bWorldShift) override;
+	// End AActor interface.
+
+#if WITH_EDITOR
+	virtual void PostEditUndo() override;
+	virtual void PostEditChangeProperty(FPropertyChangedEvent& PropertyChangedEvent) override;
+
+	// Called in response to BSP rebuilds to migrate foliage from obsolete to new components.
+	FOLIAGE_API void MapRebuild();
+
+	// Moves instances based on the specified component to the current streaming level
+	FOLIAGE_API void MoveInstancesForComponentToCurrentLevel(UActorComponent* InComponent);
+
+	// Change all instances based on one component to a new component (possible in another level).
+	// The instances keep the same world locations
+	FOLIAGE_API void MoveInstancesToNewComponent(UPrimitiveComponent* InOldComponent, UPrimitiveComponent* InNewComponent);
+
+	// Move instances based on a component that has just been moved.
+	void MoveInstancesForMovedComponent(UActorComponent* InComponent);
+
+	/** Subscribes for actors transformation events */
+	void SubscribeToUpdateEvents();
+
+	// Returns a map of Static Meshes and their placed instances attached to a component.
+	FOLIAGE_API TMap<UFoliageType*, TArray<const FFoliageInstancePlacementInfo*>> GetInstancesForComponent(UActorComponent* InComponent);
+
+	// Deletes the instances attached to a component
+	FOLIAGE_API void DeleteInstancesForComponent(UActorComponent* InComponent);
+
+	// Deletes the instances spawned by a component
+	FOLIAGE_API void DeleteInstancesForSpawner(UActorComponent* InComponent);
+
+	// Finds a mesh entry
+	FOLIAGE_API FFoliageMeshInfo* FindMesh(UFoliageType* InType);
+
+	// Finds a mesh entry or adds it if it doesn't already exist
+	FOLIAGE_API FFoliageMeshInfo* FindOrAddMesh(UFoliageType* InType);
+
+	// Add a new static mesh.
+	FOLIAGE_API FFoliageMeshInfo* AddMesh(UStaticMesh* InMesh, UFoliageType** OutSettings = nullptr);
+	FOLIAGE_API FFoliageMeshInfo* AddMesh(UFoliageType* InType);
+
+	// Remove the static mesh from the mesh list, and all its instances.
+	FOLIAGE_API void RemoveMesh(UFoliageType* InFoliageType);
+
+	// Performs a reverse lookup from a mesh to its settings
+	FOLIAGE_API UFoliageType* GetSettingsForMesh(UStaticMesh* InMesh, FFoliageMeshInfo** OutMeshInfo = nullptr);
+
+	// Select an individual instance.
+	FOLIAGE_API void SelectInstance(UInstancedStaticMeshComponent* InComponent, int32 InComponentInstanceIndex, bool bToggle);
+
+	// Propagate the selected instances to the actual render components
+	FOLIAGE_API void ApplySelectionToComponents(bool bApply);
+
+	// Updates the SelectedMesh property of the actor based on the actual selected instances
+	FOLIAGE_API void CheckSelection();
+
+	// Returns the location for the widget
+	FOLIAGE_API FVector GetSelectionLocation();
+
+	// Transforms Editor specific data which is stored in world space
+	FOLIAGE_API void ApplyLevelTransform(const FTransform& LevelTransform);
+
+	/**
+	* Get the instanced foliage actor for the current streaming level.
+	*
+	* @param InCreationWorldIfNone			World to create the foliage instance in
+	* @param bCreateIfNone					Create if doesnt already exist
+	* returns								pointer to foliage object instance
+	*/
+	static FOLIAGE_API AInstancedFoliageActor* GetInstancedFoliageActorForCurrentLevel(UWorld* InWorld, bool bCreateIfNone = true);
+
+
+	/**
+	* Get the instanced foliage actor for the specified streaming level.
+	* @param bCreateIfNone					Create if doesnt already exist
+	* returns								pointer to foliage object instance
+	*/
+	static FOLIAGE_API AInstancedFoliageActor* GetInstancedFoliageActorForLevel(ULevel* Level, bool bCreateIfNone = true);
+#endif	//WITH_EDITOR
+
+private:
+#if WITH_EDITOR
+	void OnLevelActorMoved(AActor* InActor);
+#endif
+private:
+#if WITH_EDITOR
+	FDelegateHandle OnLevelActorMovedDelegateHandle;
+#endif
+
+};
diff --git a/Engine/Source/Runtime/Foliage/Public/InteractiveFoliageActor.h b/Engine/Source/Runtime/Foliage/Public/InteractiveFoliageActor.h
new file mode 100644
--- /dev/null
+++ b/Engine/Source/Runtime/Foliage/Public/InteractiveFoliageActor.h
@@ -0,0 +1,103 @@
+// Copyright 1998-2015 Epic Games, Inc. All Rights Reserved.
+
+//=============================================================================
+//=============================================================================
+
+#pragma once
+#include "Engine/StaticMeshActor.h"
+#include "InteractiveFoliageActor.generated.h"
+
+UCLASS(MinimalAPI)
+class AInteractiveFoliageActor : public AStaticMeshActor
+{
+	GENERATED_UCLASS_BODY()
+
+private:
+	/** Collision cylinder */
+	UPROPERTY()
+	class UCapsuleComponent* CapsuleComponent;
+
+	/**
+	 * Position of the last actor to enter the collision cylinder.
+	 * This currently does not handle multiple actors affecting the foliage simultaneously.
+	 */
+	UPROPERTY(transient)
+	FVector TouchingActorEntryPosition;
+
+	/** Simulated physics state */
+	UPROPERTY(transient)
+	FVector FoliageVelocity;
+
+	/** @todo document */
+	UPROPERTY(transient)
+	FVector FoliageForce;
+
+	/** @todo document */
+	UPROPERTY(transient)
+	FVector FoliagePosition;
+
+public:
+	/** Scales forces applied from damage events. */
+	UPROPERTY(EditAnywhere, BlueprintReadWrite, Category=FoliagePhysics)
+	float FoliageDamageImpulseScale;
+
+	/** Scales forces applied from touch events. */
+	UPROPERTY(EditAnywhere, BlueprintReadWrite, Category=FoliagePhysics)
+	float FoliageTouchImpulseScale;
+
+	/** Determines how strong the force that pushes toward the spring's center will be. */
+	UPROPERTY(EditAnywhere, BlueprintReadWrite, Category=FoliagePhysics)
+	float FoliageStiffness;
+
+	/**
+	 * Same as FoliageStiffness, but the strength of this force increases with the square of the distance to the spring's center.
+	 * This force is used to prevent the spring from extending past a certain point due to touch and damage forces.
+	 */
+	UPROPERTY(EditAnywhere, BlueprintReadWrite, Category=FoliagePhysics)
+	float FoliageStiffnessQuadratic;
+
+	/**
+	 * Determines the amount of energy lost by the spring as it oscillates.
+	 * This force is similar to air friction.
+	 */
+	UPROPERTY(EditAnywhere, BlueprintReadWrite, Category=FoliagePhysics)
+	float FoliageDamping;
+
+	/** Clamps the magnitude of each damage force applied. */
+	UPROPERTY(EditAnywhere, BlueprintReadWrite, Category=FoliagePhysics)
+	float MaxDamageImpulse;
+
+	/** Clamps the magnitude of each touch force applied. */
+	UPROPERTY(EditAnywhere, BlueprintReadWrite, Category=FoliagePhysics)
+	float MaxTouchImpulse;
+
+	/** Clamps the magnitude of combined forces applied each update. */
+	UPROPERTY(EditAnywhere, BlueprintReadWrite, Category=FoliagePhysics)
+	float MaxForce;
+
+	//@todo - hook this up	/** @todo document */
+	UPROPERTY()
+	float Mass;
+
+protected:
+	/** @todo document */
+	void SetupCollisionCylinder();
+
+	/** Called when capsule is touched */
+	UFUNCTION()
+	void CapsuleTouched(AActor* Other, UPrimitiveComponent* OtherComp, int32 OtherBodyIndex, bool bFromSweep, const FHitResult& OverlapInfo);
+public:
+	// Begin AActor interface
+	virtual void Tick(float DeltaSeconds) override;
+	virtual void PostActorCreated() override;
+	virtual float TakeDamage(float DamageAmount, struct FDamageEvent const& DamageEvent, class AController* EventInstigator, class AActor* DamageCauser) override;
+	// End AActor interface
+
+	// Begin UObject interface
+	virtual void PostLoad() override;
+	// End UObject interface
+
+};
+
+
+
diff --git a/Engine/Source/Runtime/Landscape/Classes/LandscapeHeightfieldCollisionComponent.h b/Engine/Source/Runtime/Landscape/Classes/LandscapeHeightfieldCollisionComponent.h
--- a/Engine/Source/Runtime/Landscape/Classes/LandscapeHeightfieldCollisionComponent.h
+++ b/Engine/Source/Runtime/Landscape/Classes/LandscapeHeightfieldCollisionComponent.h
@@ -4,8 +4,7 @@
 #pragma once
 
 class ULandscapeLayerInfoObject;
-
-#include "Foliage/InstancedFoliageActor.h"
+class AInstancedFoliageActor;
 
 #include "LandscapeHeightfieldCollisionComponent.generated.h"
 
diff --git a/Engine/Source/Runtime/Landscape/Landscape.Build.cs b/Engine/Source/Runtime/Landscape/Landscape.Build.cs
--- a/Engine/Source/Runtime/Landscape/Landscape.Build.cs
+++ b/Engine/Source/Runtime/Landscape/Landscape.Build.cs
@@ -18,6 +18,7 @@ public class Landscape : ModuleRules
 				"TargetPlatform",
 				"DerivedDataCache",
 				"ImageWrapper",
+                "Foliage",
 			}
 		);
 
@@ -29,6 +30,7 @@ public class Landscape : ModuleRules
 				"RenderCore", 
 				"RHI",
 				"ShaderCore",
+                "Foliage",
 			}
 		);
 
diff --git a/Engine/Source/Runtime/Landscape/Private/LandscapeCollision.cpp b/Engine/Source/Runtime/Landscape/Private/LandscapeCollision.cpp
--- a/Engine/Source/Runtime/Landscape/Private/LandscapeCollision.cpp
+++ b/Engine/Source/Runtime/Landscape/Private/LandscapeCollision.cpp
@@ -19,7 +19,7 @@ LandscapeCollision.cpp: Landscape collision
 #include "LandscapeLayerInfoObject.h"
 #include "LandscapeInfo.h"
 #include "InstancedFoliage.h"
-#include "Foliage/FoliageType.h"
+#include "FoliageType.h"
 #include "Engine/StaticMesh.h"
 #include "Components/HierarchicalInstancedStaticMeshComponent.h"
 #include "Interfaces/Interface_CollisionDataProvider.h"
@@ -28,6 +28,7 @@ LandscapeCollision.cpp: Landscape collision
 #include "Engine/CollisionProfile.h"
 #include "Engine/Engine.h"
 #include "EngineGlobals.h"
+#include "InstancedFoliageActor.h"
 
 TMap<FGuid, ULandscapeHeightfieldCollisionComponent::FPhysXHeightfieldRef* > GSharedHeightfieldRefs;
 
diff --git a/Engine/Source/Runtime/Landscape/Private/LandscapeEdit.cpp b/Engine/Source/Runtime/Landscape/Private/LandscapeEdit.cpp
--- a/Engine/Source/Runtime/Landscape/Private/LandscapeEdit.cpp
+++ b/Engine/Source/Runtime/Landscape/Private/LandscapeEdit.cpp
@@ -25,7 +25,7 @@ LandscapeEdit.cpp: Landscape editing
 #include "LandscapeSplineControlPoint.h"
 #include "LandscapeSplineSegment.h"
 #include "LandscapeGizmoActiveActor.h"
-#include "Foliage/InstancedFoliageActor.h"
+#include "InstancedFoliageActor.h"
 #include "LevelUtils.h"
 #include "MessageLog.h"
 #include "MapErrors.h"
