Add a default option for AddMesh now that it's no longer in StaticMesh

diff --git a/Engine/Source/Runtime/Foliage/Private/InstancedFoliage.cpp b/Engine/Source/Runtime/Foliage/Private/InstancedFoliage.cpp
--- a/Engine/Source/Runtime/Foliage/Private/InstancedFoliage.cpp
+++ b/Engine/Source/Runtime/Foliage/Private/InstancedFoliage.cpp
@@ -1032,7 +1032,7 @@ FFoliageMeshInfo* AInstancedFoliageActor::FindOrAddMesh(UFoliageType* InType)
 	return MeshInfo;
 }
 
-FFoliageMeshInfo* AInstancedFoliageActor::AddMesh(UStaticMesh* InMesh, UFoliageType** OutSettings)
+FFoliageMeshInfo* AInstancedFoliageActor::AddMesh(UStaticMesh* InMesh, UFoliageType** OutSettings, const UFoliageType_InstancedStaticMesh* DefaultSettings)
 {
 	check(GetSettingsForMesh(InMesh) == nullptr);
 
@@ -1040,12 +1040,12 @@ FFoliageMeshInfo* AInstancedFoliageActor::AddMesh(UStaticMesh* InMesh, UFoliageT
 
 	UFoliageType_InstancedStaticMesh* Settings = nullptr;
 #if WITH_EDITORONLY_DATA
-	/*if (InMesh->FoliageDefaultSettings)
+	if (DefaultSettings)
 	{
 		// TODO: Can't we just use this directly?
-		Settings = DuplicateObject<UFoliageType_InstancedStaticMesh>(InMesh->FoliageDefaultSettings, this);
+		Settings = DuplicateObject<UFoliageType_InstancedStaticMesh>(DefaultSettings, this);
 	}
-	else*/
+	else
 #endif
 	{
 		Settings = ConstructObject<UFoliageType_InstancedStaticMesh>(UFoliageType_InstancedStaticMesh::StaticClass(), this);
diff --git a/Engine/Source/Runtime/Foliage/Public/InstancedFoliageActor.h b/Engine/Source/Runtime/Foliage/Public/InstancedFoliageActor.h
--- a/Engine/Source/Runtime/Foliage/Public/InstancedFoliageActor.h
+++ b/Engine/Source/Runtime/Foliage/Public/InstancedFoliageActor.h
@@ -13,6 +13,7 @@
 class UFoliageType;
 struct FFoliageInstancePlacementInfo;
 struct FFoliageMeshInfo;
+class UFoliageType_InstancedStaticMesh;
 
 UCLASS(notplaceable, hidecategories = Object, MinimalAPI, NotBlueprintable)
 class AInstancedFoliageActor : public AActor
@@ -76,7 +77,7 @@ public:
 	FOLIAGE_API FFoliageMeshInfo* FindOrAddMesh(UFoliageType* InType);
 
 	// Add a new static mesh.
-	FOLIAGE_API FFoliageMeshInfo* AddMesh(UStaticMesh* InMesh, UFoliageType** OutSettings = nullptr);
+	FOLIAGE_API FFoliageMeshInfo* AddMesh(UStaticMesh* InMesh, UFoliageType** OutSettings = nullptr, const UFoliageType_InstancedStaticMesh* DefaultSettings = nullptr);
 	FOLIAGE_API FFoliageMeshInfo* AddMesh(UFoliageType* InType);
 
 	// Remove the static mesh from the mesh list, and all its instances.
