Remove temp fix for instanced foliage actor where it forces mobility to be static. This isn't needed

diff --git a/Engine/Source/Runtime/Foliage/Private/InstancedFoliage.cpp b/Engine/Source/Runtime/Foliage/Private/InstancedFoliage.cpp
--- a/Engine/Source/Runtime/Foliage/Private/InstancedFoliage.cpp
+++ b/Engine/Source/Runtime/Foliage/Private/InstancedFoliage.cpp
@@ -1247,16 +1247,6 @@ void AInstancedFoliageActor::PostEditUndo()
 	}
 }
 
-void AInstancedFoliageActor::PostEditChangeProperty(FPropertyChangedEvent& PropertyChangedEvent)
-{
-	Super::PostEditChangeProperty(PropertyChangedEvent);
-
-	if (RootComponent)
-	{
-		RootComponent->SetMobility(EComponentMobility::Static);	//ensure mobility is not changed
-	}
-}
-
 void AInstancedFoliageActor::ApplySelectionToComponents(bool bApply)
 {
 	bool bNeedsUpdate = false;
@@ -1540,11 +1530,6 @@ void AInstancedFoliageActor::PostLoad()
 	}
 #endif
 
-	if (RootComponent)
-	{
-		RootComponent->SetMobility(EComponentMobility::Static);	//instanced foliage is always static
-	}
-
 #if WITH_EDITOR
 	if (GIsEditor)
 	{
diff --git a/Engine/Source/Runtime/Foliage/Public/InstancedFoliageActor.h b/Engine/Source/Runtime/Foliage/Public/InstancedFoliageActor.h
--- a/Engine/Source/Runtime/Foliage/Public/InstancedFoliageActor.h
+++ b/Engine/Source/Runtime/Foliage/Public/InstancedFoliageActor.h
@@ -43,7 +43,6 @@ public:
 
 #if WITH_EDITOR
 	virtual void PostEditUndo() override;
-	virtual void PostEditChangeProperty(FPropertyChangedEvent& PropertyChangedEvent) override;
 
 	// Called in response to BSP rebuilds to migrate foliage from obsolete to new components.
 	FOLIAGE_API void MapRebuild();
