Fix PS4 build

diff --git a/Engine/Source/Runtime/Foliage/Private/FoliageModule.cpp b/Engine/Source/Runtime/Foliage/Private/FoliageModule.cpp
new file mode 100644
--- /dev/null
+++ b/Engine/Source/Runtime/Foliage/Private/FoliageModule.cpp
@@ -0,0 +1,6 @@
+// Copyright 1998-2015 Epic Games, Inc. All Rights Reserved.
+
+#include "FoliagePrivate.h"
+#include "FoliageModule.h"
+
+IMPLEMENT_MODULE( IFoliageModule, Foliage );
diff --git a/Engine/Source/Runtime/Foliage/Private/FoliagePrivate.h b/Engine/Source/Runtime/Foliage/Private/FoliagePrivate.h
--- a/Engine/Source/Runtime/Foliage/Private/FoliagePrivate.h
+++ b/Engine/Source/Runtime/Foliage/Private/FoliagePrivate.h
@@ -2,6 +2,7 @@
 
 
 #pragma once
+#include "Engine.h"
 #include "StaticMeshResources.h"
 #include "FoliageType_InstancedStaticMesh.h"
 #include "InstancedFoliageActor.h"
diff --git a/Engine/Source/Runtime/Foliage/Public/FoliageModule.h b/Engine/Source/Runtime/Foliage/Public/FoliageModule.h
new file mode 100644
--- /dev/null
+++ b/Engine/Source/Runtime/Foliage/Public/FoliageModule.h
@@ -0,0 +1,13 @@
+// Copyright 1998-2015 Epic Games, Inc. All Rights Reserved.
+
+#pragma once
+
+#include "ModuleInterface.h"
+
+/**
+ * Foliage Edit mode module interface
+ */
+class IFoliageModule : public IModuleInterface
+{
+public:
+};
