Added level transform event and hooked foliage editor data patching to it

diff --git a/Engine/Source/Runtime/Engine/Classes/Engine/Level.h b/Engine/Source/Runtime/Engine/Classes/Engine/Level.h
--- a/Engine/Source/Runtime/Engine/Classes/Engine/Level.h
+++ b/Engine/Source/Runtime/Engine/Classes/Engine/Level.h
@@ -437,6 +437,10 @@ public:
 		return (OwningWorld && this == OwningWorld->CurrentLevelPendingVisibility);
 	}
 
+	// Event on level transform changes
+	DECLARE_MULTICAST_DELEGATE_OneParam(FLevelTransformEvent, const FTransform&);
+	FLevelTransformEvent OnApplyLevelTransform;
+
 #if WITH_EDITORONLY_DATA
 	/** Level simplification settings for each LOD */
 	UPROPERTY()
diff --git a/Engine/Source/Runtime/Engine/Private/LevelUtils.cpp b/Engine/Source/Runtime/Engine/Private/LevelUtils.cpp
--- a/Engine/Source/Runtime/Engine/Private/LevelUtils.cpp
+++ b/Engine/Source/Runtime/Engine/Private/LevelUtils.cpp
@@ -384,14 +384,9 @@ void FLevelUtils::ApplyLevelTransform( ULevel* Level, const FTransform& LevelTra
 		{
 			ApplyPostEditMove( Level );						
 		}
-
-		// We need to transform foliage Editor data which is stored in world space
-		/*AInstancedFoliageActor* FoliageActor = AInstancedFoliageActor::GetInstancedFoliageActorForLevel(Level);
-		if (FoliageActor)
-		{
-			FoliageActor->ApplyLevelTransform(LevelTransform);
-		}*/
 #endif // WITH_EDITOR
+
+		Level->OnApplyLevelTransform.Broadcast(LevelTransform);
 	}
 }
 
diff --git a/Engine/Source/Runtime/Foliage/Private/InstancedFoliage.cpp b/Engine/Source/Runtime/Foliage/Private/InstancedFoliage.cpp
--- a/Engine/Source/Runtime/Foliage/Private/InstancedFoliage.cpp
+++ b/Engine/Source/Runtime/Foliage/Private/InstancedFoliage.cpp
@@ -1670,7 +1670,32 @@ void AInstancedFoliageActor::AddReferencedObjects(UObject* InThis, FReferenceCol
 	Super::AddReferencedObjects(This, Collector);
 }
 
+void AInstancedFoliageActor::PostRegisterAllComponents()
+{
+	Super::PostRegisterAllComponents();
+
+#if WITH_EDITOR
+	if (!IsTemplate())
+	{
+		ULevel* Level = GetLevel();
+		Level->OnApplyLevelTransform.Remove(OnApplyLevelTransformDelegateHandle); // RegisterAllComponents can be called many times before calling UnregisterAllComponents
+		OnApplyLevelTransformDelegateHandle = Level->OnApplyLevelTransform.AddUObject(this, &AInstancedFoliageActor::ApplyLevelTransform);
+	}
+#endif// WITH_EDITOR
+}
 
+void AInstancedFoliageActor::PostUnregisterAllComponents()
+{
+#if WITH_EDITOR
+	if (!IsTemplate())
+	{
+		ULevel* Level = GetLevel();
+		Level->OnApplyLevelTransform.Remove(OnApplyLevelTransformDelegateHandle);
+	}
+#endif// WITH_EDITOR
+
+	Super::PostUnregisterAllComponents();
+}
 
 void AInstancedFoliageActor::ApplyWorldOffset(const FVector& InOffset, bool bWorldShift)
 {
diff --git a/Engine/Source/Runtime/Foliage/Public/InstancedFoliageActor.h b/Engine/Source/Runtime/Foliage/Public/InstancedFoliageActor.h
--- a/Engine/Source/Runtime/Foliage/Public/InstancedFoliageActor.h
+++ b/Engine/Source/Runtime/Foliage/Public/InstancedFoliageActor.h
@@ -37,7 +37,8 @@ public:
 	// Begin AActor interface.
 	// we don't want to have our components automatically destroyed by the Blueprint code
 	virtual void RerunConstructionScripts() override {}
-	// Add world origin offset
+	virtual void PostRegisterAllComponents() override;
+	virtual void PostUnregisterAllComponents() override;
 	virtual void ApplyWorldOffset(const FVector& InOffset, bool bWorldShift) override;
 	// End AActor interface.
 
@@ -125,6 +126,7 @@ private:
 private:
 #if WITH_EDITOR
 	FDelegateHandle OnLevelActorMovedDelegateHandle;
+	FDelegateHandle OnApplyLevelTransformDelegateHandle;
 #endif
 
 };
