Detect if a shared FoliageType's properties were changed while the level wasn't loaded, and if so propagte the changes to the InstancedFoliageActor.

diff --git a/Engine/Source/Editor/FoliageEdit/Private/SFoliageEditMeshDisplayItem.cpp b/Engine/Source/Editor/FoliageEdit/Private/SFoliageEditMeshDisplayItem.cpp
--- a/Engine/Source/Editor/FoliageEdit/Private/SFoliageEditMeshDisplayItem.cpp
+++ b/Engine/Source/Editor/FoliageEdit/Private/SFoliageEditMeshDisplayItem.cpp
@@ -1706,11 +1706,6 @@ void SFoliageEditMeshDisplayItem::Construct(const FArguments& InArgs)
 }
 END_SLATE_FUNCTION_BUILD_OPTIMIZATION
 
-void SFoliageEditMeshDisplayItem::NotifyPostChange(const FPropertyChangedEvent& PropertyChangedEvent, UProperty* PropertyThatChanged)
-{
-	FoliageEditPtr.Pin()->GetFoliageEditMode()->ReallocateClusters(FoliageSettingsPtr);
-}
-
 void SFoliageEditMeshDisplayItem::BindCommands()
 {
 	const FFoliageEditCommands& Commands = FFoliageEditCommands::Get();
diff --git a/Engine/Source/Editor/FoliageEdit/Private/SFoliageEditMeshDisplayItem.h b/Engine/Source/Editor/FoliageEdit/Private/SFoliageEditMeshDisplayItem.h
--- a/Engine/Source/Editor/FoliageEdit/Private/SFoliageEditMeshDisplayItem.h
+++ b/Engine/Source/Editor/FoliageEdit/Private/SFoliageEditMeshDisplayItem.h
@@ -62,9 +62,6 @@ public:
 	const FSlateBrush* GetSaveSettingsBrush() const;
 
 private:
-	/** Notification when the cluster details are changed */
-	void NotifyPostChange(const FPropertyChangedEvent& PropertyChangedEvent, UProperty* PropertyThatChanged) override;
-
 	/** Callback to get the border color for the item. */
 	FSlateColor GetBorderColor() const;
 
diff --git a/Engine/Source/Runtime/Foliage/Private/InstancedFoliage.cpp b/Engine/Source/Runtime/Foliage/Private/InstancedFoliage.cpp
--- a/Engine/Source/Runtime/Foliage/Private/InstancedFoliage.cpp
+++ b/Engine/Source/Runtime/Foliage/Private/InstancedFoliage.cpp
@@ -33,6 +33,8 @@ struct FFoliageCustomVersion
 		FoliageUsingHierarchicalISMC = 1,
 		// Changed Component to not RF_Transactional
 		HierarchicalISMCNonTransactional = 2,
+		// Added FoliageTypeUpdateGuid
+		AddedFoliageTypeUpdateGuid = 3,
 		// -----<new versions can be added above this line>-------------------------------------------------
 		VersionPlusOne,
 		LatestVersion = VersionPlusOne - 1
@@ -135,6 +137,11 @@ FArchive& operator<<(FArchive& Ar, FFoliageMeshInfo& MeshInfo)
 		Ar << MeshInfo.Instances;
 	}
 
+	if (Ar.CustomVer(FFoliageCustomVersion::GUID) >= FFoliageCustomVersion::AddedFoliageTypeUpdateGuid)
+	{
+		Ar << MeshInfo.FoliageTypeUpdateGuid;
+	}
+
 	// Serialize the transient data for undo.
 	if (Ar.IsTransacting())
 	{
@@ -223,6 +230,8 @@ UFoliageType::UFoliageType(const FObjectInitializer& ObjectInitializer)
 	FRichCurve* Curve = ScaleCurve.GetRichCurve();
 	Curve->AddKey(0.f, 0.f);
 	Curve->AddKey(1.f, 1.f);
+
+	UpdateGuid = FGuid::NewGuid();
 }
 
 
@@ -278,6 +287,17 @@ void UFoliageType::PostEditChangeProperty(struct FPropertyChangedEvent& Property
 	// Ensure that OverriddenLightMapRes is a factor of 4
 	OverriddenLightMapRes = OverriddenLightMapRes > 4 ? OverriddenLightMapRes + 3 & ~3 : 4;
 	++ChangeCount;
+
+	UpdateGuid = FGuid::NewGuid();
+
+	// Notify any currently-loaded InstancedFoliageActors
+	if (IsFoliageReallocationRequiredForPropertyChange(PropertyChangedEvent))
+	{
+		for (TObjectIterator<AInstancedFoliageActor> It(RF_ClassDefaultObject | RF_PendingKill); It; ++It)
+		{
+			It->NotifyFoliageTypeChanged(this);
+		}
+	}
 }
 #endif
 
@@ -591,6 +611,7 @@ int32 FFoliageMeshInfo::GetInstanceCount() const
 void FFoliageMeshInfo::ReallocateClusters(AInstancedFoliageActor* InIFA, UFoliageType* InSettings)
 {
 	// Detach all components
+
 	InIFA->UnregisterAllComponents();
 
 	if (Component != nullptr)
@@ -606,6 +627,9 @@ void FFoliageMeshInfo::ReallocateClusters(AInstancedFoliageActor* InIFA, UFoliag
 	ComponentHash.Empty();
 	SelectedIndices.Empty();
 
+	// Copy the UpdateGuid from the foliage type
+	FoliageTypeUpdateGuid = InSettings->UpdateGuid;
+
 	// Re-add
 	for (FFoliageInstance& Instance : OldInstances)
 	{
@@ -1111,6 +1135,7 @@ FFoliageMeshInfo* AInstancedFoliageActor::AddMesh(UFoliageType* InType)
 	}
 
 	FFoliageMeshInfo* MeshInfo = &*FoliageMeshes.Add(InType);
+	MeshInfo->FoliageTypeUpdateGuid = InType->UpdateGuid;
 	InType->IsSelected = true;
 
 	return MeshInfo;
@@ -1487,6 +1512,7 @@ void AInstancedFoliageActor::Serialize(FArchive& Ar)
 				FoliageType = (UFoliageType_InstancedStaticMesh*)StaticDuplicateObject(FoliageType, this, nullptr, RF_AllFlags & ~(RF_Standalone | RF_Public));
 				FoliageType->Mesh = OldMeshInfo.Key;
 			}
+			NewMeshInfo.FoliageTypeUpdateGuid = FoliageType->UpdateGuid;
 			FoliageMeshes.Add(FoliageType, TUniqueObj<FFoliageMeshInfo>(MoveTemp(NewMeshInfo)));
 		}
 	}
@@ -1574,6 +1600,16 @@ void AInstancedFoliageActor::PostLoad()
 				MeshInfo.ReallocateClusters(this, MeshPair.Key);
 			}
 
+			// Update foliage components if the foliage settings object was changed while the level was not loaded.
+			if (MeshInfo.FoliageTypeUpdateGuid != FoliageType->UpdateGuid)
+			{
+				if (MeshInfo.FoliageTypeUpdateGuid.IsValid())
+				{
+					MeshInfo.ReallocateClusters(this, MeshPair.Key);
+				}
+				MeshInfo.FoliageTypeUpdateGuid = FoliageType->UpdateGuid;
+			}
+
 			// Update the hash.
 			for (int32 InstanceIdx = 0; InstanceIdx < MeshInfo.Instances.Num(); InstanceIdx++)
 			{
@@ -1610,6 +1646,15 @@ void AInstancedFoliageActor::PostLoad()
 
 #if WITH_EDITOR
 
+void AInstancedFoliageActor::NotifyFoliageTypeChanged(UFoliageType* FoliageType)
+{
+	FFoliageMeshInfo* MeshInfo = FindMesh(FoliageType);
+	if (MeshInfo)
+	{
+		MeshInfo->ReallocateClusters(this, FoliageType);
+	}
+}
+
 void AInstancedFoliageActor::OnLevelActorMoved(AActor* InActor)
 {
 	if (!GetWorld()->IsPlayInEditor())
diff --git a/Engine/Source/Runtime/Foliage/Public/FoliageType.h b/Engine/Source/Runtime/Foliage/Public/FoliageType.h
--- a/Engine/Source/Runtime/Foliage/Public/FoliageType.h
+++ b/Engine/Source/Runtime/Foliage/Public/FoliageType.h
@@ -20,6 +20,11 @@ class UFoliageType : public UObject
 {
 	GENERATED_UCLASS_BODY()
 
+	/* A GUID that is updated every time the foliage type is modified, 
+	   so foliage placed in the level can detect the FoliageType has changed. */
+	UPROPERTY()
+	FGuid UpdateGuid;
+
 	UPROPERTY(EditAnywhere, Category=Painting)
 	float Density;
 
@@ -227,6 +232,9 @@ class UFoliageType : public UObject
 	virtual UStaticMesh* GetStaticMesh() PURE_VIRTUAL(UFoliageType::GetStaticMesh, return nullptr; );
 
 #if WITH_EDITOR
+	/* Lets subclasses decide if the InstancedFoliageActor should reallocate its instances if the specified property change event occurs */
+	virtual bool IsFoliageReallocationRequiredForPropertyChange(struct FPropertyChangedEvent& PropertyChangedEvent) const { return true; }
+
 	virtual void PostEditChangeProperty(struct FPropertyChangedEvent& PropertyChangedEvent) override;
 #endif
 
diff --git a/Engine/Source/Runtime/Foliage/Public/InstancedFoliage.h b/Engine/Source/Runtime/Foliage/Public/InstancedFoliage.h
--- a/Engine/Source/Runtime/Foliage/Public/InstancedFoliage.h
+++ b/Engine/Source/Runtime/Foliage/Public/InstancedFoliage.h
@@ -157,6 +157,9 @@ struct FFoliageMeshInfo
 	UHierarchicalInstancedStaticMeshComponent* Component;
 
 #if WITH_EDITORONLY_DATA
+	// Allows us to detect if FoliageType was updated while this level wasn't loaded
+	FGuid FoliageTypeUpdateGuid;
+
 	// Editor-only placed instances
 	TArray<FFoliageInstance> Instances;
 
@@ -179,6 +182,7 @@ struct FFoliageMeshInfo
 		// even VC++2013 doesn't support "=default" on move constructors
 		: Component(Other.Component)
 #if WITH_EDITORONLY_DATA
+		, FoliageTypeUpdateGuid(MoveTemp(Other.FoliageTypeUpdateGuid))
 		, Instances(MoveTemp(Other.Instances))
 		, InstanceHash(MoveTemp(Other.InstanceHash))
 		, ComponentHash(MoveTemp(Other.ComponentHash))
@@ -191,6 +195,7 @@ struct FFoliageMeshInfo
 	{
 		Component = Other.Component;
 #if WITH_EDITORONLY_DATA
+		FoliageTypeUpdateGuid = MoveTemp(Other.FoliageTypeUpdateGuid);
 		Instances = MoveTemp(Other.Instances);
 		InstanceHash = MoveTemp(Other.InstanceHash);
 		ComponentHash = MoveTemp(Other.ComponentHash);
diff --git a/Engine/Source/Runtime/Foliage/Public/InstancedFoliageActor.h b/Engine/Source/Runtime/Foliage/Public/InstancedFoliageActor.h
--- a/Engine/Source/Runtime/Foliage/Public/InstancedFoliageActor.h
+++ b/Engine/Source/Runtime/Foliage/Public/InstancedFoliageActor.h
@@ -117,6 +117,10 @@ public:
 	* returns								pointer to foliage object instance
 	*/
 	static FOLIAGE_API AInstancedFoliageActor* GetInstancedFoliageActorForLevel(ULevel* Level, bool bCreateIfNone = true);
+
+	/* Called to notify InstancedFoliageActor that a UFoliageType has been modified */
+	void NotifyFoliageTypeChanged(UFoliageType* FoliageType);
+
 #endif	//WITH_EDITOR
 
 private:
