Fix for WITH_EDITORONLY_DATA

diff --git a/Engine/Source/Runtime/Foliage/Private/InstancedFoliage.cpp b/Engine/Source/Runtime/Foliage/Private/InstancedFoliage.cpp
--- a/Engine/Source/Runtime/Foliage/Private/InstancedFoliage.cpp
+++ b/Engine/Source/Runtime/Foliage/Private/InstancedFoliage.cpp
@@ -1512,7 +1512,9 @@ void AInstancedFoliageActor::Serialize(FArchive& Ar)
 				FoliageType = (UFoliageType_InstancedStaticMesh*)StaticDuplicateObject(FoliageType, this, nullptr, RF_AllFlags & ~(RF_Standalone | RF_Public));
 				FoliageType->Mesh = OldMeshInfo.Key;
 			}
+#if WITH_EDITORONLY_DATA
 			NewMeshInfo.FoliageTypeUpdateGuid = FoliageType->UpdateGuid;
+#endif
 			FoliageMeshes.Add(FoliageType, TUniqueObj<FFoliageMeshInfo>(MoveTemp(NewMeshInfo)));
 		}
 	}
