Rename Ecosystem to ProceduralFoliage

diff --git a/Engine/Source/Editor/EditorStyle/Private/SlateEditorStyle.cpp b/Engine/Source/Editor/EditorStyle/Private/SlateEditorStyle.cpp
--- a/Engine/Source/Editor/EditorStyle/Private/SlateEditorStyle.cpp
+++ b/Engine/Source/Editor/EditorStyle/Private/SlateEditorStyle.cpp
@@ -4005,10 +4005,6 @@ void FSlateEditorStyle::FStyle::SetupLevelEditorStyle()
 		Set( "LevelEditor.FoliageMode.Selected", new IMAGE_BRUSH( "Icons/icon_Mode_Foliage_selected_40x", Icon40x40 ) );
 		Set( "LevelEditor.FoliageMode.Selected.Small", new IMAGE_BRUSH( "Icons/icon_Mode_Foliage_selected_40x", Icon20x20 ) );
 
-		Set("LevelEditor.EcosystemMode", new IMAGE_BRUSH("Icons/icon_Mode_Ecosystem_40x", Icon40x40));
-		Set("LevelEditor.EcosystemMode.Small", new IMAGE_BRUSH("Icons/icon_Mode_Ecosystem_40x", Icon20x20));
-
-
 		Set( "LevelEditor.BspMode", new IMAGE_BRUSH( "Icons/icon_Mode_GeoEdit_40px", Icon40x40 ) );
 		Set( "LevelEditor.BspMode.Small", new IMAGE_BRUSH( "Icons/icon_Mode_GeoEdit_40px", Icon20x20 ) );
 		Set( "LevelEditor.BspMode.Selected", new IMAGE_BRUSH( "Icons/icon_Mode_GeoEdit-a_40px", Icon40x40 ) );
diff --git a/Engine/Source/Editor/FoliageEdit/FoliageEdit.build.cs b/Engine/Source/Editor/FoliageEdit/FoliageEdit.build.cs
--- a/Engine/Source/Editor/FoliageEdit/FoliageEdit.build.cs
+++ b/Engine/Source/Editor/FoliageEdit/FoliageEdit.build.cs
@@ -20,6 +20,7 @@ public class FoliageEdit : ModuleRules
 				"LevelEditor",
 				"Landscape",
                 "PropertyEditor",
+                "AssetTools",
                 "Foliage",
 			}
 		);
diff --git a/Engine/Source/Editor/FoliageEdit/Private/ActorFactoryProceduralFoliage.cpp b/Engine/Source/Editor/FoliageEdit/Private/ActorFactoryProceduralFoliage.cpp
new file mode 100644
--- /dev/null
+++ b/Engine/Source/Editor/FoliageEdit/Private/ActorFactoryProceduralFoliage.cpp
@@ -0,0 +1,78 @@
+// Copyright 1998-2015 Epic Games, Inc. All Rights Reserved.
+
+/*=============================================================================
+ActorFactory.cpp:
+=============================================================================*/
+
+#include "UnrealEd.h"
+#include "ActorFactoryProceduralFoliage.h"
+#include "ProceduralFoliage.h"
+#include "ProceduralFoliageActor.h"
+#include "ProceduralFoliageComponent.h"
+
+DEFINE_LOG_CATEGORY_STATIC(LogActorFactory, Log, All);
+
+#define LOCTEXT_NAMESPACE "ActorFactoryProceduralFoliage"
+
+/*-----------------------------------------------------------------------------
+UActorFactoryProceduralFoliage
+-----------------------------------------------------------------------------*/
+UActorFactoryProceduralFoliage::UActorFactoryProceduralFoliage(const FObjectInitializer& ObjectInitializer)
+: Super(ObjectInitializer)
+{
+	DisplayName = LOCTEXT("ProceduralFoliageDisplayName", "ProceduralFoliage");
+	NewActorClass = AProceduralFoliageActor::StaticClass();
+	bUseSurfaceOrientation = true;
+}
+
+bool UActorFactoryProceduralFoliage::CanCreateActorFrom(const FAssetData& AssetData, FText& OutErrorMsg)
+{
+	if (!AssetData.IsValid() || !AssetData.GetClass()->IsChildOf(UProceduralFoliage::StaticClass()))
+	{
+		OutErrorMsg = NSLOCTEXT("CanCreateActor", "NoProceduralFoliage", "A valid ProceduralFoliage must be specified.");
+		return false;
+	}
+
+	return true;
+}
+
+void UActorFactoryProceduralFoliage::PostSpawnActor(UObject* Asset, AActor* NewActor)
+{
+	UProceduralFoliage* ProceduralFoliage = CastChecked<UProceduralFoliage>(Asset);
+	GEditor->SetActorLabelUnique(NewActor, ProceduralFoliage->GetName());
+
+	UE_LOG(LogActorFactory, Log, TEXT("Actor Factory created %s"), *ProceduralFoliage->GetName());
+
+	// Change properties
+	AProceduralFoliageActor* PFA = CastChecked<AProceduralFoliageActor>(NewActor);
+	UProceduralFoliageComponent* ProceduralComponent = PFA->ProceduralComponent;
+	check(ProceduralComponent);
+
+	ProceduralComponent->UnregisterComponent();
+
+	ProceduralComponent->ProceduralFoliage = ProceduralFoliage;
+
+	// Init Component
+	ProceduralComponent->RegisterComponent();
+}
+
+UObject* UActorFactoryProceduralFoliage::GetAssetFromActorInstance(AActor* Instance)
+{
+	check(Instance->IsA(NewActorClass));
+	AProceduralFoliageActor* ProceduralFoliage = CastChecked<AProceduralFoliageActor>(Instance);
+	UProceduralFoliageComponent* ProceduralComponent = ProceduralFoliage->ProceduralComponent;
+	check(ProceduralComponent);
+	return ProceduralComponent->ProceduralFoliage;
+}
+
+void UActorFactoryProceduralFoliage::PostCreateBlueprint(UObject* Asset, AActor* CDO)
+{
+	if (Asset != NULL && CDO != NULL)
+	{
+		UProceduralFoliage* ProceduralFoliage = CastChecked<UProceduralFoliage>(Asset);
+		AProceduralFoliageActor* PFA = CastChecked<AProceduralFoliageActor>(CDO);
+		UProceduralFoliageComponent* ProceduralComponent = PFA->ProceduralComponent;
+		ProceduralComponent->ProceduralFoliage = ProceduralFoliage;
+	}
+}
+#undef LOCTEXT_NAMESPACE
diff --git a/Engine/Source/Editor/FoliageEdit/Private/ActorFactoryProceduralFoliage.h b/Engine/Source/Editor/FoliageEdit/Private/ActorFactoryProceduralFoliage.h
new file mode 100644
--- /dev/null
+++ b/Engine/Source/Editor/FoliageEdit/Private/ActorFactoryProceduralFoliage.h
@@ -0,0 +1,19 @@
+// Copyright 1998-2015 Epic Games, Inc. All Rights Reserved.
+
+
+#pragma once
+#include "ActorFactories/ActorFactory.h"
+#include "ActorFactoryProceduralFoliage.generated.h"
+
+UCLASS(MinimalAPI, config=Editor)
+class UActorFactoryProceduralFoliage : public UActorFactory
+{
+	GENERATED_UCLASS_BODY()
+
+	// Begin UActorFactory Interface
+	virtual bool CanCreateActorFrom( const FAssetData& AssetData, FText& OutErrorMsg ) override;
+	virtual void PostSpawnActor( UObject* Asset, AActor* NewActor) override;
+	virtual void PostCreateBlueprint( UObject* Asset, AActor* CDO ) override;
+	virtual UObject* GetAssetFromActorInstance(AActor* ActorInstance) override;
+	// End UActorFactory Interface
+};
diff --git a/Engine/Source/Editor/FoliageEdit/Private/FoliageEditModule.cpp b/Engine/Source/Editor/FoliageEdit/Private/FoliageEditModule.cpp
--- a/Engine/Source/Editor/FoliageEdit/Private/FoliageEditModule.cpp
+++ b/Engine/Source/Editor/FoliageEdit/Private/FoliageEditModule.cpp
@@ -9,6 +9,11 @@ const FName FoliageEditAppIdentifier = FName(TEXT("FoliageEdApp"));
 #include "FoliageEdMode.h"
 #include "PropertyEditing.h"
 #include "FoliageTypeDetails.h"
+#include "ProceduralFoliageComponent.h"
+#include "ProceduralFoliageComponentVisualizer.h"
+#include "ProceduralFoliageComponentDetails.h"
+#include "ActorFactoryProceduralFoliage.h"
+#include "ComponentVisualizer.h"
 
 /**
  * Foliage Edit Mode module
@@ -32,6 +37,15 @@ public:
 		// Register the details customizer
 		FPropertyEditorModule& PropertyModule = FModuleManager::LoadModuleChecked<FPropertyEditorModule>("PropertyEditor");
 		PropertyModule.RegisterCustomClassLayout("FoliageType", FOnGetDetailCustomizationInstance::CreateStatic(&FFoliageTypeDetails::MakeInstance));
+
+		GUnrealEd->RegisterComponentVisualizer(UProceduralFoliageComponent::StaticClass()->GetFName(), MakeShareable(new FProceduralFoliageComponentVisualizer));
+
+		FPropertyEditorModule& PropertyEditor = FModuleManager::LoadModuleChecked<FPropertyEditorModule>("PropertyEditor");
+		PropertyEditor.RegisterCustomClassLayout("ProceduralFoliageComponent", FOnGetDetailCustomizationInstance::CreateStatic(&FProceduralFoliageComponentDetails::MakeInstance));
+
+		// Actor Factories
+		UActorFactoryProceduralFoliage* ProceduralFoliageActorFactory = ConstructObject<UActorFactoryProceduralFoliage>(UActorFactoryProceduralFoliage::StaticClass());
+		GEditor->ActorFactories.Add(ProceduralFoliageActorFactory);
 	}
 
 	/**
diff --git a/Engine/Source/Editor/FoliageEdit/Private/FoliageEditPrivate.h b/Engine/Source/Editor/FoliageEdit/Private/FoliageEditPrivate.h
new file mode 100644
--- /dev/null
+++ b/Engine/Source/Editor/FoliageEdit/Private/FoliageEditPrivate.h
@@ -0,0 +1,4 @@
+// Copyright 1998-2015 Epic Games, Inc. All Rights Reserved.
+
+#include "UnrealEd.h"
+#include "PropertyEditing.h"
diff --git a/Engine/Source/Editor/FoliageEdit/Private/FoliageTypeDetails.cpp b/Engine/Source/Editor/FoliageEdit/Private/FoliageTypeDetails.cpp
--- a/Engine/Source/Editor/FoliageEdit/Private/FoliageTypeDetails.cpp
+++ b/Engine/Source/Editor/FoliageEdit/Private/FoliageTypeDetails.cpp
@@ -48,6 +48,6 @@ void AddSubcategoryProperties(IDetailCategoryBuilder& CategoryBuilder)
 void FFoliageTypeDetails::CustomizeDetails(IDetailLayoutBuilder& DetailBuilder)
 {
 	AddSubcategoryProperties(DetailBuilder.EditCategory("General"));
-	AddSubcategoryProperties(DetailBuilder.EditCategory("Ecosystem"));
+	AddSubcategoryProperties(DetailBuilder.EditCategory("Procedural"));
 	AddSubcategoryProperties(DetailBuilder.EditCategory("Painting"));
 }
\ No newline at end of file
diff --git a/Engine/Source/Editor/FoliageEdit/Private/ProceduralFoliageComponentDetails.cpp b/Engine/Source/Editor/FoliageEdit/Private/ProceduralFoliageComponentDetails.cpp
new file mode 100644
--- /dev/null
+++ b/Engine/Source/Editor/FoliageEdit/Private/ProceduralFoliageComponentDetails.cpp
@@ -0,0 +1,91 @@
+// Copyright 1998-2015 Epic Games, Inc. All Rights Reserved.
+
+#include "UnrealEd.h"
+#include "ProceduralFoliageComponentDetails.h"
+#include "ProceduralFoliage.h"
+#include "ProceduralFoliageComponent.h"
+
+TSharedRef<IDetailCustomization> FProceduralFoliageComponentDetails::MakeInstance()
+{
+	return MakeShareable( new FProceduralFoliageComponentDetails() );
+
+}
+
+
+void FProceduralFoliageComponentDetails::CustomizeDetails( IDetailLayoutBuilder& DetailBuilder )
+{
+	const FName ProceduralFoliageCategoryName("ProceduralFoliage");
+	IDetailCategoryBuilder& ProceduralFoliageCategory = DetailBuilder.EditCategory(ProceduralFoliageCategoryName);
+
+	const FText ResimulateText = NSLOCTEXT("ProceduralFoliageComponentDetails","ResimulateButtonText", "Resimulate" );
+
+	TArray< TWeakObjectPtr<UObject> > ObjectsBeingCustomized;
+	DetailBuilder.GetObjectsBeingCustomized( ObjectsBeingCustomized );
+
+	for( TWeakObjectPtr<UObject>& Object : ObjectsBeingCustomized )
+	{
+		UProceduralFoliageComponent* Component = Cast<UProceduralFoliageComponent>( Object.Get() );
+		if( ensure( Component ) )
+		{
+			SelectedComponents.Add( Component );
+		}
+	}
+
+	TArray<TSharedRef<IPropertyHandle>> AllProperties;
+	bool bSimpleProperties = true;
+	bool bAdvancedProperties = false;
+	// Add all properties in the category in order
+	ProceduralFoliageCategory.GetDefaultProperties(AllProperties, bSimpleProperties, bAdvancedProperties);
+	for( auto& Property : AllProperties )
+	{
+		ProceduralFoliageCategory.AddProperty(Property);
+	}
+
+	FDetailWidgetRow& NewRow =  ProceduralFoliageCategory.AddCustomRow( ResimulateText );
+
+	NewRow.ValueContent()
+	.MaxDesiredWidth(120.f)
+	[
+		SNew(SButton)
+		.OnClicked( this, &FProceduralFoliageComponentDetails::OnResimulateClicked )
+		.ToolTipText( NSLOCTEXT("ProceduralFoliageComponentDetails","ResimulateButton_Tooltip", "Resimulates the ProceduralFoliage asset and replaces previously spawned instances" ) )
+		.IsEnabled( this, &FProceduralFoliageComponentDetails::IsResimulateEnabled )
+		[
+			SNew( STextBlock )
+			.Font( IDetailLayoutBuilder::GetDetailFont() )
+			.Text( ResimulateText )
+		]
+	];
+}
+
+FReply FProceduralFoliageComponentDetails::OnResimulateClicked()
+{
+	TSet<UProceduralFoliage*> UniqueProceduralFoliages;
+	for( TWeakObjectPtr<UProceduralFoliageComponent>& Component : SelectedComponents )
+	{
+		if( Component.IsValid() && Component->ProceduralFoliage )
+		{
+			if( !UniqueProceduralFoliages.Contains( Component->ProceduralFoliage ) )
+			{
+				UniqueProceduralFoliages.Add( Component->ProceduralFoliage );
+				Component->ProceduralFoliage->Simulate();
+			}
+
+			Component->SpawnProceduralContent();
+		}
+	}
+	return FReply::Handled();
+}
+
+bool FProceduralFoliageComponentDetails::IsResimulateEnabled() const
+{
+	for(const TWeakObjectPtr<UProceduralFoliageComponent>& Component : SelectedComponents)
+	{
+		if(Component.IsValid() && Component->ProceduralFoliage)
+		{
+			return true;
+		}
+	}
+
+	return false;
+}
\ No newline at end of file
diff --git a/Engine/Source/Editor/FoliageEdit/Private/ProceduralFoliageComponentDetails.h b/Engine/Source/Editor/FoliageEdit/Private/ProceduralFoliageComponentDetails.h
new file mode 100644
--- /dev/null
+++ b/Engine/Source/Editor/FoliageEdit/Private/ProceduralFoliageComponentDetails.h
@@ -0,0 +1,21 @@
+// Copyright 1998-2015 Epic Games, Inc. All Rights Reserved.
+
+#pragma once
+
+class FProceduralFoliageComponentDetails : public IDetailCustomization
+{
+public:
+	virtual ~FProceduralFoliageComponentDetails(){};
+
+	/** Makes a new instance of this detail layout class for a specific detail view requesting it */
+	static TSharedRef<IDetailCustomization> MakeInstance();
+
+	/** IDetailCustomization interface */
+	virtual void CustomizeDetails( IDetailLayoutBuilder& DetailBuilder );
+private:
+	FReply OnResimulateClicked();
+	bool IsResimulateEnabled() const;
+private:
+	TArray< TWeakObjectPtr<class UProceduralFoliageComponent> > SelectedComponents;
+
+};
diff --git a/Engine/Source/Editor/FoliageEdit/Private/ProceduralFoliageComponentVisualizer.cpp b/Engine/Source/Editor/FoliageEdit/Private/ProceduralFoliageComponentVisualizer.cpp
new file mode 100644
--- /dev/null
+++ b/Engine/Source/Editor/FoliageEdit/Private/ProceduralFoliageComponentVisualizer.cpp
@@ -0,0 +1,39 @@
+// Copyright 1998-2015 Epic Games, Inc. All Rights Reserved.
+
+#include "UnrealEd.h"
+#include "ProceduralFoliageComponentVisualizer.h"
+#include "ProceduralFoliage.h"
+#include "ProceduralFoliageComponent.h"
+
+
+static const FColor	ProcTileColor(0,255,0);
+static const FColor	ProcTileOverlapColor(255, 255, 0);
+
+void FProceduralFoliageComponentVisualizer::DrawVisualization(const UActorComponent* Component, const FSceneView* View, FPrimitiveDrawInterface* PDI)
+{
+	if (const UProceduralFoliageComponent* ProcComponent = Cast<const UProceduralFoliageComponent>(Component))
+	{
+		if (ProcComponent->ProceduralFoliage)
+		{
+			const float TileSize = ProcComponent->ProceduralFoliage->TileSize;
+			const FVector TileSizeOffset(TileSize, TileSize, 0.f);
+			const FVector TileLocation = ProcComponent->GetComponentToWorld().GetLocation();
+			const FVector OverlapV(ProcComponent->Overlap, ProcComponent->Overlap, 0.f);
+			
+			for (int32 X = 0; X < ProcComponent->TilesX; ++X)
+			{
+				for (int32 Y = 0; Y < ProcComponent->TilesY; ++Y)
+				{
+					const FVector StartOffset(X*TileSize, Y*TileSize, 0.f);
+					
+					const FBox BoxInner(TileLocation + StartOffset, TileLocation + StartOffset + TileSizeOffset);
+					DrawWireBox(PDI, BoxInner, ProcTileColor, SDPG_World);
+
+					const FBox BoxOuter = BoxInner + (BoxInner.Min - OverlapV) + (BoxInner.Max + OverlapV);
+					DrawWireBox(PDI, BoxOuter, ProcTileOverlapColor, SDPG_World);
+				}
+			}
+			
+		}
+	}
+}
diff --git a/Engine/Source/Editor/FoliageEdit/Private/ProceduralFoliageComponentVisualizer.h b/Engine/Source/Editor/FoliageEdit/Private/ProceduralFoliageComponentVisualizer.h
new file mode 100644
--- /dev/null
+++ b/Engine/Source/Editor/FoliageEdit/Private/ProceduralFoliageComponentVisualizer.h
@@ -0,0 +1,13 @@
+// Copyright 1998-2015 Epic Games, Inc. All Rights Reserved.
+
+#pragma once
+
+#include "ComponentVisualizer.h"
+
+class FProceduralFoliageComponentVisualizer : public FComponentVisualizer
+{
+public:
+	// Begin FComponentVisualizer interface
+	virtual void DrawVisualization(const UActorComponent* Component, const FSceneView* View, FPrimitiveDrawInterface* PDI) override;
+	// End FComponentVisualizer interface
+};
diff --git a/Engine/Source/Editor/FoliageEdit/Private/ProceduralFoliageFactory.cpp b/Engine/Source/Editor/FoliageEdit/Private/ProceduralFoliageFactory.cpp
new file mode 100644
--- /dev/null
+++ b/Engine/Source/Editor/FoliageEdit/Private/ProceduralFoliageFactory.cpp
@@ -0,0 +1,26 @@
+// Copyright 1998-2015 Epic Games, Inc. All Rights Reserved.
+
+#include "UnrealEd.h"
+#include "ProceduralFoliageFactory.h"
+#include "IAssetTypeActions.h"
+#include "ProceduralFoliage.h"
+
+UProceduralFoliageFactory::UProceduralFoliageFactory(const FObjectInitializer& ObjectInitializer)
+	: Super(ObjectInitializer)
+{
+	bCreateNew = true;
+	bEditAfterNew = true;
+	SupportedClass = UProceduralFoliage::StaticClass();
+}
+
+UObject* UProceduralFoliageFactory::FactoryCreateNew(UClass* Class, UObject* InParent, FName Name, EObjectFlags Flags, UObject* Context, FFeedbackContext* Warn)
+{
+	UProceduralFoliage* NewProceduralFoliage = ConstructObject<UProceduralFoliage>(Class, InParent, Name, Flags | RF_Transactional);
+
+	return NewProceduralFoliage;
+}
+
+uint32 UProceduralFoliageFactory::GetMenuCategories() const
+{
+	return EAssetTypeCategories::Foliage;
+}
\ No newline at end of file
diff --git a/Engine/Source/Editor/FoliageEdit/Private/ProceduralFoliageFactory.h b/Engine/Source/Editor/FoliageEdit/Private/ProceduralFoliageFactory.h
new file mode 100644
--- /dev/null
+++ b/Engine/Source/Editor/FoliageEdit/Private/ProceduralFoliageFactory.h
@@ -0,0 +1,21 @@
+// Copyright 1998-2015 Epic Games, Inc. All Rights Reserved.
+
+/**
+ * Factory for ProceduralFoliage assets
+ */
+
+#pragma once
+
+#include "ProceduralFoliageFactory.generated.h"
+
+
+UCLASS()
+class UProceduralFoliageFactory : public UFactory
+{
+	GENERATED_UCLASS_BODY()
+
+	// UFactory interface
+	virtual UObject* FactoryCreateNew(UClass* Class, UObject* InParent, FName Name, EObjectFlags Flags, UObject* Context, FFeedbackContext* Warn) override;
+	virtual uint32 GetMenuCategories() const;
+	// End of UFactory interface
+};
diff --git a/Engine/Source/Editor/FoliageEdit/Private/SFoliageEditMeshDisplayItem.cpp b/Engine/Source/Editor/FoliageEdit/Private/SFoliageEditMeshDisplayItem.cpp
--- a/Engine/Source/Editor/FoliageEdit/Private/SFoliageEditMeshDisplayItem.cpp
+++ b/Engine/Source/Editor/FoliageEdit/Private/SFoliageEditMeshDisplayItem.cpp
@@ -1956,7 +1956,7 @@ bool SFoliageEditMeshDisplayItem::IsPropertyVisible(const FPropertyAndParent& Pr
 	const FName CategoryName = FObjectEditorUtils::GetCategoryFName(&PropertyAndParent.Property);
 	
 	FString Subcategory = "";
-	if ( (CategoryName != FName("Ecosystem")) && PropertyAndParent.Property.HasMetaData(FName("Subcategory")))
+	if ( (CategoryName != FName("Procedural")) && PropertyAndParent.Property.HasMetaData(FName("Subcategory")))
 	{
 		Subcategory = PropertyAndParent.Property.GetMetaData(FName("Subcategory"));
 	}
diff --git a/Engine/Source/Runtime/Foliage/Private/ProceduralFoliage.cpp b/Engine/Source/Runtime/Foliage/Private/ProceduralFoliage.cpp
new file mode 100644
--- /dev/null
+++ b/Engine/Source/Runtime/Foliage/Private/ProceduralFoliage.cpp
@@ -0,0 +1,164 @@
+// Copyright 1998-2015 Epic Games, Inc. All Rights Reserved.
+
+#include "FoliagePrivate.h"
+#include "ProceduralFoliage.h"
+#include "ProceduralFoliageTile.h"
+#include "Components/BoxComponent.h"
+#include "PhysicsPublic.h"
+
+#define LOCTEXT_NAMESPACE "ProceduralFoliage"
+
+UProceduralFoliage::UProceduralFoliage(const FObjectInitializer& ObjectInitializer)
+	: Super(ObjectInitializer)
+{
+	TileSize = 10000;	//100 m
+	NumUniqueTiles = 10;
+	RandomSeed = 42;
+}
+
+#if WITH_EDITOR
+void UProceduralFoliage::PostEditChangeProperty(FPropertyChangedEvent& PropertyChangedEvent)
+{
+	bNeedsSimulation = true;
+}
+#endif
+
+UProceduralFoliageTile* UProceduralFoliage::CreateTempTile()
+{
+	UProceduralFoliageTile* TmpTile = NewObject<UProceduralFoliageTile>(this);
+	TmpTile->InitSimulation(this);
+
+	return TmpTile;
+}
+
+void UProceduralFoliage::GenerateTile(const int32 NumSteps)
+{
+	UProceduralFoliageTile* NewTile = NewObject<UProceduralFoliageTile>(this);
+	PrecomputedTiles.Add(NewTile);
+
+	NewTile->Simulate(this, NumSteps);
+}
+
+void UProceduralFoliage::CreateProceduralFoliageInstances()
+{
+	for(FProceduralFoliageTypeData& TypeData : Types)
+	{
+		if( TypeData.Type )
+		{
+			TypeData.TypeInstance = ConstructObject<UFoliageType_InstancedStaticMesh>(TypeData.Type, this);
+			TypeData.ChangeCount = TypeData.TypeInstance->ChangeCount;
+		}
+	}
+}
+
+void UProceduralFoliage::SetClean()
+{
+	for (FProceduralFoliageTypeData& TypeData : Types)
+	{
+		if( TypeData.Type && TypeData.TypeInstance )
+		{
+			TypeData.TypeInstance->ChangeCount = TypeData.Type->GetDefaultObject<UFoliageType_InstancedStaticMesh>()->ChangeCount;
+		}
+	}
+
+	bNeedsSimulation = false;
+}
+
+// Custom serialization version for all packages containing ProceduralFoliage
+struct FProceduralFoliageCustomVersion
+{
+	enum Type
+	{
+		// Before any version changes were made in the plugin
+		BeforeCustomVersionWasAdded = 0,
+		// We serialize the type clean change map to ensure validity
+		SerializeTypeMap = 1,
+		// -----<new versions can be added above this line>-------------------------------------------------
+		VersionPlusOne,
+		LatestVersion = VersionPlusOne - 1
+	};
+
+	// The GUID for this custom version number
+	const static FGuid GUID;
+
+private:
+	FProceduralFoliageCustomVersion() {}
+};
+
+const FGuid FProceduralFoliageCustomVersion::GUID(0xaafe32bd, 0x53954c14, 0xb66a5e25, 0x1032d1dd);
+// Register the custom version with core
+FCustomVersionRegistration GRegisterProceduralFoliageCustomVersion(FProceduralFoliageCustomVersion::GUID, FProceduralFoliageCustomVersion::LatestVersion, TEXT("ProceduralFoliageVer"));
+
+void UProceduralFoliage::Serialize(FArchive& Ar)
+{
+	Super::Serialize(Ar);
+
+	Ar.UsingCustomVersion(FProceduralFoliageCustomVersion::GUID);
+}
+
+bool UProceduralFoliage::AnyDirty() const
+{
+	bool bDirty = false;
+
+	if( Types.Num() > 0 )
+	{
+		bDirty = bNeedsSimulation;
+	for (const FProceduralFoliageTypeData& TypeData : Types)
+	{
+		if (TypeData.Type)
+		{
+			if (TypeData.TypeInstance && TypeData.Type->GetDefaultObject<UFoliageType_InstancedStaticMesh>()->ChangeCount != TypeData.ChangeCount)
+			{
+					bDirty = true;
+					break;
+			}
+			else if (!TypeData.TypeInstance)
+			{
+					bDirty = true;
+					break;
+				}
+			}
+		}
+	}
+
+	return bDirty;
+}
+
+const UProceduralFoliageTile* UProceduralFoliage::GetRandomTile(int32 X, int32 Y)
+{
+	if (AnyDirty())
+	{
+		Simulate();
+	}
+
+	if (PrecomputedTiles.Num())
+	{
+		//const int32 Idx = ((Y + X) * (Y+X)) % PrecomputedTiles.Num();
+		const int32 Idx = ((Y + X) * (Y+X)) % PrecomputedTiles.Num();
+		return PrecomputedTiles[Idx];
+	}
+
+	return nullptr;
+}
+
+void UProceduralFoliage::Simulate(int32 NumSteps)
+{
+	RandomStream.Initialize(RandomSeed);
+	CreateProceduralFoliageInstances();
+
+	PrecomputedTiles.Empty();
+
+	for (int i = 0; i < NumUniqueTiles; ++i)
+	{
+		GenerateTile(NumSteps);
+	}
+
+	SetClean();
+}
+
+int32 UProceduralFoliage::GetRandomNumber()
+{
+	return RandomStream.FRand() * RAND_MAX;
+}
+
+#undef LOCTEXT_NAMESPACE
\ No newline at end of file
diff --git a/Engine/Source/Runtime/Foliage/Private/ProceduralFoliageActor.cpp b/Engine/Source/Runtime/Foliage/Private/ProceduralFoliageActor.cpp
new file mode 100644
--- /dev/null
+++ b/Engine/Source/Runtime/Foliage/Private/ProceduralFoliageActor.cpp
@@ -0,0 +1,12 @@
+// Copyright 1998-2015 Epic Games, Inc. All Rights Reserved.
+
+#include "FoliagePrivate.h"
+#include "ProceduralFoliageActor.h"
+#include "ProceduralFoliageComponent.h"
+
+AProceduralFoliageActor::AProceduralFoliageActor(const FObjectInitializer& ObjectInitializer)
+: Super(ObjectInitializer)
+{
+	ProceduralComponent = ObjectInitializer.CreateDefaultSubobject<UProceduralFoliageComponent>(this, TEXT("ProceduralFoliageComponent"));
+	RootComponent = ProceduralComponent;
+}
diff --git a/Engine/Source/Runtime/Foliage/Private/ProceduralFoliageBlockingVolume.cpp b/Engine/Source/Runtime/Foliage/Private/ProceduralFoliageBlockingVolume.cpp
new file mode 100644
--- /dev/null
+++ b/Engine/Source/Runtime/Foliage/Private/ProceduralFoliageBlockingVolume.cpp
@@ -0,0 +1,17 @@
+// Copyright 1998-2015 Epic Games, Inc. All Rights Reserved.
+
+#include "FoliagePrivate.h"
+#include "Components/BrushComponent.h"
+#include "ProceduralFoliageBlockingVolume.h"
+
+static FName ProceduralFoliageBlocking_NAME(TEXT("ProceduralFoliageBlockingVolume"));
+
+AProceduralFoliageBlockingVolume::AProceduralFoliageBlockingVolume(const FObjectInitializer& ObjectInitializer)
+	: Super(ObjectInitializer)
+{
+	if (UBrushComponent* BrushComponent = GetBrushComponent())
+	{
+		BrushComponent->SetCollisionObjectType(ECC_WorldStatic);
+		BrushComponent->SetCollisionResponseToAllChannels(ECR_Ignore);
+	}
+}
diff --git a/Engine/Source/Runtime/Foliage/Private/ProceduralFoliageBlockingVolume.h b/Engine/Source/Runtime/Foliage/Private/ProceduralFoliageBlockingVolume.h
new file mode 100644
--- /dev/null
+++ b/Engine/Source/Runtime/Foliage/Private/ProceduralFoliageBlockingVolume.h
@@ -0,0 +1,18 @@
+// Copyright 1998-2015 Epic Games, Inc. All Rights Reserved.
+
+
+#pragma once
+#include "GameFramework/Volume.h"
+#include "ProceduralFoliageBlockingVolume.generated.h"
+
+/** An invisible volume used to block ProceduralFoliage instances from being spawned. */
+UCLASS(MinimalAPI)
+class AProceduralFoliageBlockingVolume : public AVolume
+{
+	GENERATED_UCLASS_BODY()
+
+	// End UObject interface.
+};
+
+
+
diff --git a/Engine/Source/Runtime/Foliage/Private/ProceduralFoliageBroadphase.cpp b/Engine/Source/Runtime/Foliage/Private/ProceduralFoliageBroadphase.cpp
new file mode 100644
--- /dev/null
+++ b/Engine/Source/Runtime/Foliage/Private/ProceduralFoliageBroadphase.cpp
@@ -0,0 +1,92 @@
+// Copyright 1998-2015 Epic Games, Inc. All Rights Reserved.
+
+#include "FoliagePrivate.h"
+#include "ProceduralFoliageBroadphase.h"
+#include "ProceduralFoliageInstance.h"
+
+FProceduralFoliageBroadphase::FProceduralFoliageBroadphase(float TileSize)
+: QuadTree(FBox2D(FVector2D(0, 0), FVector2D(TileSize, TileSize)))
+{
+}
+
+FProceduralFoliageBroadphase::FProceduralFoliageBroadphase(const FProceduralFoliageBroadphase& OtherBroadphase)
+: QuadTree(FBox2D(FVector2D(0, 0), FVector2D(0, 0)))
+{
+	OtherBroadphase.QuadTree.Duplicate(&QuadTree);
+}
+
+/*Takes the instance and returns the AABB that contains both its shade and collision radius*/
+FBox2D GetMaxAABB(FProceduralFoliageInstance* Instance)
+{
+	const float Radius = Instance->GetMaxRadius();
+	const FVector2D Location(Instance->Location);
+	const FVector2D Offset(Radius, Radius);
+	const FBox2D AABB(Location - Offset, Location + Offset);
+	return AABB;
+}
+
+void FProceduralFoliageBroadphase::Insert(FProceduralFoliageInstance* Instance)
+{
+	const FBox2D MaxAABB = GetMaxAABB(Instance);
+	QuadTree.Insert(Instance, MaxAABB);
+}
+
+bool CircleOverlap(const FVector& ALocation, float ARadius, const FVector& BLocation, float BRadius)
+{
+	return (ALocation - BLocation).SizeSquared2D() <= (ARadius + BRadius)*(ARadius + BRadius);
+}
+
+bool FProceduralFoliageBroadphase::GetOverlaps(FProceduralFoliageInstance* Instance, TArray<FProceduralFoliageOverlap>& Overlaps) const
+{
+	const float AShadeRadius     = Instance->GetShadeRadius();
+	const float ACollisionRadius = Instance->GetCollisionRadius();
+
+	TArray<FProceduralFoliageInstance*> PossibleOverlaps;
+	const FBox2D AABB = GetMaxAABB(Instance);
+	QuadTree.GetElementsUnique(AABB, PossibleOverlaps);
+	Overlaps.Reserve(Overlaps.Num() + PossibleOverlaps.Num());
+	
+	for (FProceduralFoliageInstance* Overlap : PossibleOverlaps)
+	{
+		if (Overlap != Instance)
+		{
+			//We must determine if this is an overlap of shade or and overlap of collision. If both the collision overlap wins
+			bool bCollisionOverlap = CircleOverlap(Instance->Location, ACollisionRadius, Overlap->Location, Overlap->GetCollisionRadius());
+			bool bShadeOverlap     = CircleOverlap(Instance->Location, AShadeRadius, Overlap->Location, Overlap->GetShadeRadius());
+
+			if (bCollisionOverlap || bShadeOverlap)
+			{
+				new (Overlaps)FProceduralFoliageOverlap(Instance, Overlap, bCollisionOverlap ? ESimulationOverlap::CollisionOverlap : ESimulationOverlap::ShadeOverlap);
+			}
+			
+		}
+	}
+
+	return Overlaps.Num() > 0;
+}
+
+void FProceduralFoliageBroadphase::Remove(FProceduralFoliageInstance* Instance)
+{
+	const FBox2D AABB = GetMaxAABB(Instance);
+	QuadTree.Remove(Instance, AABB);
+}
+
+void FProceduralFoliageBroadphase::GetInstancesInBox(const FBox2D& Box, TArray<FProceduralFoliageInstance*>& Instances) const
+{
+	TArray<FProceduralFoliageInstance*> AABBInstances;
+	QuadTree.GetElementsUnique(Box, AABBInstances);
+
+	Instances.Reserve(AABBInstances.Num() + Instances.Num());
+
+	for (FProceduralFoliageInstance* Inst : AABBInstances)
+	{
+		const FVector2D Location2D(Inst->Location.X, Inst->Location.Y);
+		const FVector2D ClosestPt = Box.GetClosestPointTo(Location2D);
+		const float Radius2 = (ClosestPt - Location2D).SizeSquared();
+		const float InstRadius = Inst->GetMaxRadius();
+		if (Radius2 < InstRadius*InstRadius)
+		{
+			Instances.Add(Inst);
+		}
+	}
+}
\ No newline at end of file
diff --git a/Engine/Source/Runtime/Foliage/Private/ProceduralFoliageBroadphase.h b/Engine/Source/Runtime/Foliage/Private/ProceduralFoliageBroadphase.h
new file mode 100644
--- /dev/null
+++ b/Engine/Source/Runtime/Foliage/Private/ProceduralFoliageBroadphase.h
@@ -0,0 +1,20 @@
+// Copyright 1998-2015 Epic Games, Inc. All Rights Reserved.
+
+#pragma once
+#include "GenericQuadTree.h"
+#include "ProceduralFoliageInstance.h"
+
+class FProceduralFoliageBroadphase
+{
+public:
+	FProceduralFoliageBroadphase(float TileSize = 0.f);
+	FProceduralFoliageBroadphase(const FProceduralFoliageBroadphase& OtherBroadphase);
+
+	void Insert(FProceduralFoliageInstance* Instance);
+	bool GetOverlaps(FProceduralFoliageInstance* Instance, TArray<FProceduralFoliageOverlap>& Overlaps) const;
+	void GetInstancesInBox(const FBox2D& Box, TArray<FProceduralFoliageInstance*>& Instances) const;
+	void Remove(FProceduralFoliageInstance* Instance);
+
+private:
+	TQuadTree<FProceduralFoliageInstance*, 4> QuadTree;
+};
\ No newline at end of file
diff --git a/Engine/Source/Runtime/Foliage/Private/ProceduralFoliageComponent.cpp b/Engine/Source/Runtime/Foliage/Private/ProceduralFoliageComponent.cpp
new file mode 100644
--- /dev/null
+++ b/Engine/Source/Runtime/Foliage/Private/ProceduralFoliageComponent.cpp
@@ -0,0 +1,225 @@
+// Copyright 1998-2015 Epic Games, Inc. All Rights Reserved.
+
+#include "FoliagePrivate.h"
+#include "ProceduralFoliageComponent.h"
+#include "ProceduralFoliageInstance.h"
+#include "ProceduralFoliageTile.h"
+#include "InstancedFoliage.h"
+#include "InstancedFoliageActor.h"
+#include "ProceduralFoliage.h"
+#include "Async/Async.h"
+
+UProceduralFoliageComponent::UProceduralFoliageComponent(const FObjectInitializer& ObjectInitializer)
+: Super(ObjectInitializer)
+{
+	TilesX = 1;
+	TilesY = 1;
+	Overlap = 0.f;
+	HalfHeight = 10000.f;
+}
+
+void SpawnInstances(const TArray<FProceduralFoliageInstance>& ProceduralFoliageInstances, UWorld* World, UActorComponent* BaseComponent)
+{
+#if WITH_EDITOR
+	FFoliageInstance Inst;
+	
+	for (const FProceduralFoliageInstance& EcoInst : ProceduralFoliageInstances)
+	{
+		if (EcoInst.BaseComponent)
+		{
+			ULevel* OwningLevel = EcoInst.BaseComponent->GetOwner()->GetLevel();
+		
+			AInstancedFoliageActor* IFA = AInstancedFoliageActor::GetInstancedFoliageActorForLevel(OwningLevel);
+			const UFoliageType_InstancedStaticMesh* Type = EcoInst.Type;
+			FFoliageMeshInfo* MeshInfo;
+			if (UStaticMesh* StaticMesh = EcoInst.StaticMesh)
+			{
+				if (IFA->GetSettingsForMesh(StaticMesh) == nullptr)
+				{
+					IFA->AddMesh(StaticMesh, nullptr, Type);
+				}
+
+				UFoliageType* Settings = IFA->GetSettingsForMesh(StaticMesh, &MeshInfo);
+				if (Settings && MeshInfo)
+				{
+					/* Convert ProceduralFoliage instance into foliage instance*/
+					Inst.DrawScale3D = FVector(EcoInst.Scale);
+					Inst.Location = EcoInst.Location;
+					Inst.Rotation = EcoInst.Rotation.Rotator();
+					Inst.Flags |= FOLIAGE_NoRandomYaw;
+
+					if (Type->AlignToNormal)
+					{
+						Inst.AlignToNormal(EcoInst.Normal, Type->AlignMaxAngle);
+					}
+
+
+					Inst.Base = EcoInst.BaseComponent;
+					Inst.Spawner = BaseComponent;
+
+
+					MeshInfo->AddInstance(IFA, Settings, Inst);
+				}
+			}
+		}
+	}
+#endif
+}
+
+void CopyTileInstances(const UProceduralFoliageTile* FromTile, UProceduralFoliageTile* ToTile, const FBox2D& InnerLocalAABB, const FTransform& ToLocalTM, const float Overlap)
+{
+	const FBox2D OuterLocalAABB(InnerLocalAABB.Min, InnerLocalAABB.Max + FVector2D(Overlap, Overlap));
+	TArray<FProceduralFoliageInstance*> ToInstances;
+	FromTile->GetInstancesInAABB(OuterLocalAABB, ToInstances, false);
+	ToTile->AddInstances(ToInstances, ToLocalTM, InnerLocalAABB);
+}
+
+FBox2D GetTileRegion(const int32 X, const int32 Y, const int32 CountX, const int32 CountY, const float InnerSize, const float Overlap)
+{
+	const FVector2D Left(-Overlap, Overlap);
+	const FVector2D Bottom(Overlap, -Overlap);
+
+	FBox2D Region(FVector2D(Overlap,Overlap), FVector2D(InnerSize + Overlap, InnerSize + Overlap));
+	if (X == 0)
+	{
+		Region += Left;
+	}
+
+	if (Y == 0)
+	{
+		Region += Bottom;
+	}
+
+	return Region;
+}
+
+FBox2D GetRightOverlap(const FBox2D& Region, const float InnerSize, const float Overlap)
+{
+	return Region + FVector2D(InnerSize + Overlap, 0.f);
+}
+
+FBox2D GetLeftOverlap(const FBox2D& Region, const float Overlap)
+{
+	return Region + FVector2D(-Overlap, 0.f);
+}
+
+FBox2D GetAboveOverlap(const FBox2D& Region, const float InnerSize, const float Overlap)
+{
+	return Region + FVector2D(0.f, InnerSize + Overlap);
+}
+
+FBox2D GetBelowOverlap(const FBox2D& Region, const float Overlap)
+{
+	return Region + FVector2D(0.f, -Overlap);
+}
+
+void UProceduralFoliageComponent::SpawnTiles()
+{
+#if WITH_EDITOR
+	if (ProceduralFoliage)
+	{
+		/** Constants for laying out the overlapping tile grid*/
+		const float InnerTileSize = ProceduralFoliage->TileSize;
+		const FVector2D InnerTileV(InnerTileSize, InnerTileSize);
+		const FVector2D InnerTileX(InnerTileSize, 0.f);
+		const FVector2D InnerTileY(0.f, InnerTileSize);
+
+		const float OuterTileSize = ProceduralFoliage->TileSize + Overlap;
+		const FVector2D OuterTileV(OuterTileSize, OuterTileSize);
+		const FVector2D OuterTileX(OuterTileSize, 0.f);
+		const FVector2D OuterTileY(0.f, OuterTileSize);
+
+		const FVector2D OverlapX(Overlap, 0.f);
+		const FVector2D OverlapY(0.f, Overlap);
+
+		TArray<TFuture< TArray<FProceduralFoliageInstance>* >> Futures;
+
+		for (int32 X = 0; X < TilesX; ++X)
+		{
+			for (int32 Y = 0; Y < TilesY; ++Y)
+			{
+				
+				//We have to get the tiles and create new one to build on main thread
+				const UProceduralFoliageTile* Tile = ProceduralFoliage->GetRandomTile(X, Y);
+				const UProceduralFoliageTile* RightTile = (X + 1 < TilesX) ? ProceduralFoliage->GetRandomTile(X + 1, Y) : nullptr;
+				const UProceduralFoliageTile* TopTile = (Y + 1 < TilesY) ? ProceduralFoliage->GetRandomTile(X, Y+1) : nullptr;
+				const UProceduralFoliageTile* TopRightTile = (RightTile && TopTile) ? ProceduralFoliage->GetRandomTile(X + 1, Y+1) : nullptr;
+
+				UProceduralFoliageTile* JTile = ProceduralFoliage->CreateTempTile();
+
+				Futures.Add(Async<TArray<FProceduralFoliageInstance>*>(EAsyncExecution::TaskGraph, [=]()
+				{
+					FTransform TileTM = ComponentToWorld;
+					const FVector OrientedOffset = ComponentToWorld.TransformVectorNoScale(FVector(X, Y, 0.f) * FVector(InnerTileSize, InnerTileSize, 0.f));
+					TileTM.AddToTranslation(OrientedOffset);
+
+					//copy the inner tile
+					const FBox2D InnerBox = GetTileRegion(X, Y, TilesX, TilesY, InnerTileSize, Overlap);
+					CopyTileInstances(Tile, JTile, InnerBox, FTransform::Identity, Overlap);
+
+
+					if (RightTile)
+					{
+						//Add overlap to the right
+						const FBox2D RightBox(FVector2D(-Overlap, InnerBox.Min.Y), FVector2D(Overlap, InnerBox.Max.Y));
+						const FTransform RightTM(FVector(InnerTileSize, 0.f, 0.f));
+						CopyTileInstances(RightTile, JTile, RightBox, RightTM, Overlap);
+					}
+
+					if (TopTile)
+					{
+						//Add overlap to the top
+						const FBox2D TopBox(FVector2D(InnerBox.Min.X, -Overlap), FVector2D(InnerBox.Max.X, Overlap));
+						const FTransform TopTM(FVector(0.f, InnerTileSize, 0.f));
+						CopyTileInstances(TopTile, JTile, TopBox, TopTM, Overlap);
+					}
+
+					if (TopRightTile)
+					{
+						//Add overlap to the top right
+						const FBox2D TopRightBox(FVector2D(-Overlap, -Overlap), FVector2D(Overlap, Overlap));
+						const FTransform TopRightTM(FVector(InnerTileSize, InnerTileSize, 0.f));
+						CopyTileInstances(TopRightTile, JTile, TopRightBox, TopRightTM, Overlap);
+					}
+
+					TArray<FProceduralFoliageInstance>* ProceduralFoliageInstances = new TArray<FProceduralFoliageInstance>();
+					JTile->InstancesToArray();
+					JTile->CreateInstancesToSpawn(*ProceduralFoliageInstances, TileTM, World, HalfHeight);
+
+					return ProceduralFoliageInstances;
+					
+				})
+				);
+			}
+		}
+
+		int32 FutureIdx = 0;
+		for (int X = 0; X < TilesX; ++X)
+		{
+			for (int Y = 0; Y < TilesY; ++Y)
+			{
+				TArray<FProceduralFoliageInstance>* ProceduralFoliageInstances = Futures[FutureIdx++].Get();
+				SpawnInstances(*ProceduralFoliageInstances, World, this);
+				delete ProceduralFoliageInstances;
+			}
+		}
+	}
+#endif
+}
+
+void UProceduralFoliageComponent::SpawnProceduralContent()
+{
+#if WITH_EDITOR
+	UWorld* World = GetWorld();
+	for (ULevel* Level : World->GetLevels())
+	{
+		if (Level)
+		{
+			AInstancedFoliageActor* IFA = AInstancedFoliageActor::GetInstancedFoliageActorForLevel(Level);
+			IFA->DeleteInstancesForSpawner(this);
+		}
+	}
+	
+	SpawnTiles();
+#endif
+}
\ No newline at end of file
diff --git a/Engine/Source/Runtime/Foliage/Private/ProceduralFoliageInstance.cpp b/Engine/Source/Runtime/Foliage/Private/ProceduralFoliageInstance.cpp
new file mode 100644
--- /dev/null
+++ b/Engine/Source/Runtime/Foliage/Private/ProceduralFoliageInstance.cpp
@@ -0,0 +1,107 @@
+	// Copyright 1998-2015 Epic Games, Inc. All Rights Reserved.
+
+#include "FoliagePrivate.h"
+#include "ProceduralFoliageInstance.h"
+#include "FoliageType_InstancedStaticMesh.h"
+
+FProceduralFoliageInstance::FProceduralFoliageInstance()
+: Location(ForceInit)
+, Rotation(ForceInit)
+, Normal(ForceInit)
+, Age(0)
+, Scale(1)
+, Type(nullptr)
+, StaticMesh(nullptr)
+, bBlocker(false)
+, BaseComponent(nullptr)
+, bAlive(true)
+{
+}
+
+FProceduralFoliageInstance::FProceduralFoliageInstance(const FProceduralFoliageInstance& Other)
+: Location(Other.Location)
+, Rotation(Other.Rotation)
+, Normal(Other.Normal)
+, Age(Other.Age)
+, Scale(Other.Scale)
+, Type(Other.Type)
+, StaticMesh(Other.StaticMesh)
+, bBlocker(Other.bBlocker)
+, BaseComponent(Other.BaseComponent)
+, bAlive(Other.bAlive)
+{
+
+}
+
+FProceduralFoliageInstance* GetLessFit(FProceduralFoliageInstance* A, FProceduralFoliageInstance* B)
+{
+	//Blocker is used for culling instances when we overlap tiles. It always wins
+	if (A->bBlocker){ return B; }
+	if (B->bBlocker){ return A; }
+	
+	//we look at priority, then age, then radius
+	if (A->Type->OverlapPriority == B->Type->OverlapPriority)
+	{
+		if (A->Age == B->Age)
+		{
+			return A->Scale < B->Scale ? A : B;
+		}
+		else
+		{
+			return A->Age < B->Age ? A : B;
+		}
+	}
+	else
+	{
+		return A->Type->OverlapPriority < B->Type->OverlapPriority ? A : B;
+	}
+}
+
+
+FProceduralFoliageInstance* FProceduralFoliageInstance::Domination(FProceduralFoliageInstance* A, FProceduralFoliageInstance* B, ESimulationOverlap::Type OverlapType)
+{
+	const UFoliageType_InstancedStaticMesh* AType = A->Type;
+	const UFoliageType_InstancedStaticMesh* BType = B->Type;
+
+	FProceduralFoliageInstance* Dominated = GetLessFit(A, B);
+
+	if (OverlapType == ESimulationOverlap::ShadeOverlap)
+	{
+		const UFoliageType_InstancedStaticMesh* DominatedType = Dominated->Type;
+		if (DominatedType->bGrowsInShade)
+		{
+			return nullptr;
+		}
+	}
+
+	return Dominated;
+}
+
+float FProceduralFoliageInstance::GetMaxRadius() const
+{
+	const float CollisionRadius = GetCollisionRadius();
+	const float ShadeRadius = GetShadeRadius();
+	return FMath::Max(CollisionRadius, ShadeRadius);
+}
+
+float FProceduralFoliageInstance::GetShadeRadius() const
+{
+	const float ShadeRadius = Type->ShadeRadius * Scale;
+	return ShadeRadius;
+}
+
+float FProceduralFoliageInstance::GetCollisionRadius() const
+{
+	const float CollisionRadius = Type->CollisionRadius * Scale;
+	return CollisionRadius;
+}
+
+void FProceduralFoliageInstance::TerminateInstance()
+{
+	bAlive = false;
+}
+
+bool FProceduralFoliageInstance::GrowsInShade() const 
+{ 
+	return Type->bGrowsInShade; 
+}
diff --git a/Engine/Source/Runtime/Foliage/Private/ProceduralFoliageTile.cpp b/Engine/Source/Runtime/Foliage/Private/ProceduralFoliageTile.cpp
new file mode 100644
--- /dev/null
+++ b/Engine/Source/Runtime/Foliage/Private/ProceduralFoliageTile.cpp
@@ -0,0 +1,436 @@
+// Copyright 1998-2015 Epic Games, Inc. All Rights Reserved.
+
+#include "FoliagePrivate.h"
+#include "ProceduralFoliageTile.h"
+#include "ProceduralFoliage.h"
+#include "ProceduralFoliageBroadphase.h"
+#include "ProceduralFoliageBlockingVolume.h"
+#include "InstancedFoliageActor.h"
+
+#define LOCTEXT_NAMESPACE "ProceduralFoliage"
+
+UProceduralFoliageTile::UProceduralFoliageTile(const FObjectInitializer& ObjectInitializer)
+: Super(ObjectInitializer)
+{
+}
+
+
+bool UProceduralFoliageTile::HandleOverlaps(FProceduralFoliageInstance* Instance)
+{
+	//We now check what the instance overlaps. If any of its overlaps dominate we remove the instance and leave everything else alone.
+	//If the instance survives we mark all dominated overlaps as pending removal. They will be removed from the broadphase and will not spread seeds or age.
+	//Note that this introduces potential indeterminism! If the iteration order changes we could get different results. This is needed because it gives us huge performance savings.
+	//Note that if the underlying data structures stay the same (i.e. no core engine changes) this should not matter. This gives us short term determinism, but not long term.
+
+	bool bSurvived = true;
+	TArray<FProceduralFoliageOverlap> Overlaps;
+	Broadphase.GetOverlaps(Instance, Overlaps);
+
+	//Check if the instance survives
+	for (const FProceduralFoliageOverlap& Overlap : Overlaps)
+	{
+		FProceduralFoliageInstance* Dominated = FProceduralFoliageInstance::Domination(Overlap.A, Overlap.B, Overlap.OverlapType);
+		if (Dominated == Instance)
+		{
+			bSurvived = false;
+			break;
+		}
+	}
+
+	if (bSurvived)
+	{
+		for (const FProceduralFoliageOverlap& Overlap : Overlaps)
+		{
+			if (FProceduralFoliageInstance* Dominated = FProceduralFoliageInstance::Domination(Overlap.A, Overlap.B, Overlap.OverlapType))
+			{
+				check(Dominated != Instance);	//we shouldn't be here if we survived
+				MarkPendingRemoval(Dominated);	//We can't immediately remove because we're potentially iterating over existing instances.
+			}
+		}
+	}
+	else
+	{
+		//didn't survive so just die
+		MarkPendingRemoval(Instance);
+	}
+
+	return bSurvived;
+}
+
+FProceduralFoliageInstance* UProceduralFoliageTile::NewSeed(const FVector& Location, float Scale, const UFoliageType_InstancedStaticMesh* Type, float InAge, bool bBlocker)
+{
+	const float InitRadius = Type->GetMaxRadius() * Scale;
+	{
+		FProceduralFoliageInstance* NewInst = new FProceduralFoliageInstance();
+		NewInst->Location = Location;
+		NewInst->Rotation = FQuat(FVector(0, 0, 1), RandomStream.FRandRange(0, 2.f*PI));
+		NewInst->Age = InAge;
+		NewInst->Type = Type;
+		NewInst->Normal = FVector(0, 0, 1);
+		NewInst->Scale = Scale;
+		NewInst->bBlocker = bBlocker;
+		
+		if (Type->StaticMeshes.Num())
+		{
+			int32 MeshIdx = RandomStream.RandRange(0, Type->StaticMeshes.Num()-1);
+			NewInst->StaticMesh = Type->StaticMeshes[MeshIdx];
+		}
+		
+
+		Broadphase.Insert(NewInst);
+		const bool bSurvived = HandleOverlaps(NewInst);
+		return bSurvived ? NewInst : nullptr;
+	}
+
+	return nullptr;
+}
+
+float GetSeedMinDistance(const FProceduralFoliageInstance* Instance, const float NewInstanceAge, const int32 SimulationStep)
+{
+	const UFoliageType_InstancedStaticMesh* Type = Instance->Type;
+	const int32 StepsLeft = Type->MaxAge - SimulationStep;
+	const float InstanceMaxAge = Type->GetNextAge(Instance->Age, StepsLeft);
+	const float NewInstanceMaxAge = Type->GetNextAge(NewInstanceAge, StepsLeft);
+
+	const float InstanceMaxScale = Type->GetScaleForAge(InstanceMaxAge);
+	const float NewInstanceMaxScale = Type->GetScaleForAge(NewInstanceMaxAge);
+
+	const float InstanceMaxRadius = InstanceMaxScale * Type->GetMaxRadius();
+	const float NewInstanceMaxRadius = NewInstanceMaxScale * Type->GetMaxRadius();
+
+	return InstanceMaxRadius + NewInstanceMaxRadius;
+}
+
+/** Generates a random number with a normal distribution with mean=0 and variance = 1. Uses Box-Muller transformation http://mathworld.wolfram.com/Box-MullerTransformation.html */
+float UProceduralFoliageTile::GetRandomGaussian()
+{
+	const float Rand1 = FMath::Max<float>(RandomStream.FRand(), SMALL_NUMBER);
+	const float Rand2 = FMath::Max<float>(RandomStream.FRand(), SMALL_NUMBER);
+	const float SqrtLn = FMath::Sqrt(-2.f * FMath::Loge(Rand1));
+	const float Rand2TwoPi = Rand2 * 2.f * PI;
+	const float Z1 = SqrtLn * FMath::Cos(Rand2TwoPi);
+	return Z1;
+}
+
+FVector UProceduralFoliageTile::GetSeedOffset(const UFoliageType_InstancedStaticMesh* Type, float MinDistance)
+{
+	//We want 10% of seeds to be the max distance so we use a z score of +- 1.64
+	const float MaxZScore = 1.64f;
+	const float Z1 = GetRandomGaussian();
+	const float Z1Clamped = FMath::Clamp(Z1, -MaxZScore, MaxZScore);
+	const float VariationDistance = Z1Clamped * Type->SpreadVariance / MaxZScore;
+	const float AverageDistance = MinDistance + Type->AverageSpreadDistance;
+	
+	const float RandRad = FMath::Max<float>(RandomStream.FRand(), SMALL_NUMBER) * PI * 2.f;
+	const FVector Dir = FVector(FMath::Cos(RandRad), FMath::Sin(RandRad), 0);
+	return Dir * (AverageDistance + VariationDistance);
+}
+
+void UProceduralFoliageTile::AgeSeeds()
+{
+	TArray<FProceduralFoliageInstance*> NewSeeds;
+	for (FProceduralFoliageInstance* Instance : Instances)
+	{
+		if (Instance->IsAlive())
+		{
+			const UFoliageType_InstancedStaticMesh* Type = Instance->Type;
+			if (SimulationStep <= Type->NumSteps)
+			{
+				const float CurrentAge = Instance->Age;
+				const float NewAge = Type->GetNextAge(Instance->Age, 1);
+				const float NewScale = Type->GetScaleForAge(NewAge);
+
+				const FVector Location = Instance->Location;
+
+				MarkPendingRemoval(Instance);
+				if (FProceduralFoliageInstance* Inst = NewSeed(Location, NewScale, Type, NewAge))
+				{
+					NewSeeds.Add(Inst);
+				}
+			}
+		}
+	}
+
+	for (FProceduralFoliageInstance* Seed : NewSeeds)
+	{
+		Instances.Add(Seed);
+	}
+
+	FlushPendingRemovals();
+	
+}
+
+void UProceduralFoliageTile::SpreadSeeds(TArray<FProceduralFoliageInstance*>& NewSeeds)
+{
+	for (FProceduralFoliageInstance* Inst : Instances)
+	{
+		if (Inst->IsAlive() == false)	//The instance has been killed so don't bother spreading seeds. Note this introduces potential indeterminism if the order of instance traversal changes (implementation details of TSet for example)
+		{
+			continue;
+		}
+
+		const UFoliageType_InstancedStaticMesh* Type = Inst->Type;
+
+		if (SimulationStep <= Type->NumSteps)
+		{
+			for (int32 i = 0; i < Type->SeedsPerStep; ++i)
+			{
+				//spread new seeds
+				const float NewAge = Type->GetInitAge(RandomStream);
+				const float NewScale = Type->GetScaleForAge(NewAge);
+				const float MinDistanceToClear = GetSeedMinDistance(Inst, NewAge, SimulationStep);
+				const FVector GlobalOffset = GetSeedOffset(Type, MinDistanceToClear);
+				
+				if (GlobalOffset.SizeSquared2D() + SMALL_NUMBER > MinDistanceToClear*MinDistanceToClear)
+				{
+					const FVector NewLocation = GlobalOffset + Inst->Location;
+					if (FProceduralFoliageInstance* Inst = NewSeed(NewLocation, NewScale, Type, NewAge))
+					{
+						NewSeeds.Add(Inst);
+					}
+				}
+			}
+		}
+	}
+}
+
+void UProceduralFoliageTile::AddRandomSeeds(TArray<FProceduralFoliageInstance*>& OutInstances)
+{
+	const float SizeTenM2 = (ProceduralFoliage->TileSize * ProceduralFoliage->TileSize) / (1000.f * 1000.f);
+
+	for ( const FProceduralFoliageTypeData& Data : ProceduralFoliage->GetTypes() )
+	{
+		UFoliageType_InstancedStaticMesh* TypeInstance = Data.TypeInstance;
+		if( TypeInstance )
+		{
+			const float NumSeeds = TypeInstance->GetSeedDensitySquared() * SizeTenM2;
+			for(int32 i = 0; i < NumSeeds; ++i)
+			{
+				const float X = RandomStream.FRandRange(0, ProceduralFoliage->TileSize);
+				const float Y = RandomStream.FRandRange(0, ProceduralFoliage->TileSize);
+				const float NewAge = TypeInstance->GetInitAge(RandomStream);
+				const float Scale = TypeInstance->GetScaleForAge(NewAge);
+
+				if(FProceduralFoliageInstance* NewInst = NewSeed(FVector(X,Y,0.f), Scale, TypeInstance, NewAge))
+				{
+					OutInstances.Add(NewInst);
+				}
+			}
+		}
+	}
+}
+
+void UProceduralFoliageTile::MarkPendingRemoval(FProceduralFoliageInstance* ToRemove)
+{
+	if (ToRemove->IsAlive())
+	{
+		Broadphase.Remove(ToRemove);	//we can remove from broadphase right away
+		ToRemove->TerminateInstance();
+		PendingRemovals.Add(ToRemove);
+	}
+}
+
+void UProceduralFoliageTile::RemoveInstances()
+{
+	for (FProceduralFoliageInstance* Inst : Instances)
+	{
+		MarkPendingRemoval(Inst);
+	}
+
+	InstancesArray.Empty();
+	FlushPendingRemovals();
+}
+
+void UProceduralFoliageTile::InstancesToArray()
+{
+	InstancesArray.Empty(Instances.Num());
+	for (FProceduralFoliageInstance* FromInst : Instances)
+	{
+		if (FromInst->bBlocker == false)	//blockers do not get instantiated so don't bother putting it into array
+		{
+			new(InstancesArray)FProceduralFoliageInstance(*FromInst);
+		}
+	}
+}
+
+void UProceduralFoliageTile::RemoveInstance(FProceduralFoliageInstance* ToRemove)
+{
+	if (ToRemove->IsAlive())
+	{
+		Broadphase.Remove(ToRemove);
+		ToRemove->TerminateInstance();
+	}
+	
+	Instances.Remove(ToRemove);
+	delete ToRemove;
+}
+
+void UProceduralFoliageTile::FlushPendingRemovals()
+{
+	for (FProceduralFoliageInstance* ToRemove : PendingRemovals)
+	{
+		RemoveInstance(ToRemove);
+	}
+
+	PendingRemovals.Empty();
+}
+
+void UProceduralFoliageTile::InitSimulation(UProceduralFoliage* InProceduralFoliage)
+{
+	
+
+	RandomStream.Initialize(InProceduralFoliage->GetRandomNumber());
+	ProceduralFoliage = InProceduralFoliage;
+	SimulationStep = 0;
+	Broadphase = FProceduralFoliageBroadphase(ProceduralFoliage->TileSize);
+}
+
+void UProceduralFoliageTile::StepSimulation()
+{
+	TArray<FProceduralFoliageInstance*> NewInstances;
+	if (SimulationStep == 0)
+	{
+		AddRandomSeeds(NewInstances);
+	}
+	else
+	{
+		AgeSeeds();
+		SpreadSeeds(NewInstances);
+	}
+
+	for (FProceduralFoliageInstance* Inst : NewInstances)
+	{
+		Instances.Add(Inst);
+	}
+
+	FlushPendingRemovals();
+}
+
+void UProceduralFoliageTile::Simulate(UProceduralFoliage* InProceduralFoliage, const int32 MaxNumSteps)
+{
+	InitSimulation(InProceduralFoliage);
+
+	int32 MaxSteps = 0;
+	for(const FProceduralFoliageTypeData& Data : ProceduralFoliage->GetTypes())
+	{
+		const UFoliageType_InstancedStaticMesh* TypeInstance = Data.TypeInstance;
+		if(TypeInstance)
+		{
+			MaxSteps = FMath::Max(MaxSteps, TypeInstance->NumSteps+1);
+		}
+	}
+
+	if (MaxNumSteps >= 0)
+	{
+		MaxSteps = FMath::Min(MaxSteps, MaxNumSteps);	//only take as many steps as given
+	}
+
+	FScopedSlowTask SlowTask(MaxSteps, LOCTEXT("SimulateProceduralFoliage", "Simulate ProceduralFoliage..."));
+	SlowTask.MakeDialog();
+	
+	for (int32 Step = 0; Step < MaxSteps; ++Step)
+	{
+		StepSimulation();
+		++SimulationStep;
+		SlowTask.EnterProgressFrame(1);
+	}
+
+	InstancesToArray();
+}
+
+
+void UProceduralFoliageTile::BeginDestroy()
+{
+	Super::BeginDestroy();
+	RemoveInstances();
+}
+
+void UProceduralFoliageTile::CreateInstancesToSpawn(TArray<FProceduralFoliageInstance>& OutInstances, const FTransform& WorldTM, UWorld* World, const float HalfHeight) const
+{
+	const FCollisionQueryParams Params(true);
+	FHitResult Hit;
+
+	OutInstances.Reserve(Instances.Num());
+	for (const FProceduralFoliageInstance& Instance : InstancesArray)
+	{
+		//@todo ProceduralFoliage need a better method for calculating the ray
+		FVector StartRay = Instance.Location + WorldTM.GetLocation();
+		StartRay.Z += HalfHeight;
+		FVector EndRay = StartRay;
+		EndRay.Z -= HalfHeight*2.f;
+		FCollisionShape SphereShape;
+		SphereShape.SetSphere(Instance.GetMaxRadius());
+
+		if (World->SweepSingle(Hit, StartRay, EndRay, FQuat::Identity, SphereShape, Params, FCollisionObjectQueryParams(ECC_WorldStatic)))
+		{
+			if (Hit.Actor.IsValid())	//if we hit the ProceduralFoliage blocking volume don't spawn instance
+			{
+				if (Cast<AProceduralFoliageBlockingVolume>(Hit.Actor.Get()) || Cast<AInstancedFoliageActor>(Hit.Actor.Get()))
+				{
+					continue;
+				}
+			}
+
+			const UFoliageType_InstancedStaticMesh* Type = Instance.Type;
+			if (Hit.ImpactPoint.Z >= Type->HeightMin && Hit.ImpactPoint.Z <= Type->HeightMax)
+			{
+				if (FMath::Cos(FMath::DegreesToRadians(Type->GroundSlope)) <= Hit.ImpactNormal.Z)
+				{
+					FProceduralFoliageInstance* NewInst = new(OutInstances)FProceduralFoliageInstance(Instance);
+					NewInst->Location = FVector(StartRay.X, StartRay.Y, Hit.ImpactPoint.Z);	//take the x,y of the instance, but use the z of the impact point. This is because we never want to move the instance along xy or we'll get overlaps 
+					NewInst->Normal = Hit.ImpactNormal;
+					if (Hit.Component.IsValid())
+					{
+						NewInst->BaseComponent = Hit.Component.Get();
+					}
+				}
+			}
+		}
+	}
+}
+
+void UProceduralFoliageTile::GetInstancesInAABB(const FBox2D& LocalAABB, TArray<FProceduralFoliageInstance*>& OutInstances, bool bOnTheBorder) const
+{
+	TArray<FProceduralFoliageInstance*> InstancesInAABB;
+	Broadphase.GetInstancesInBox(LocalAABB, InstancesInAABB);
+
+	OutInstances.Reserve(OutInstances.Num() + InstancesInAABB.Num());
+	for (FProceduralFoliageInstance* Inst : InstancesInAABB)
+{
+		const float Rad = Inst->GetMaxRadius();
+		const FVector& Location = Inst->Location;
+
+		if (bOnTheBorder || (Location.X - Rad >= LocalAABB.Min.X && Location.X + Rad <= LocalAABB.Max.X && Location.Y - Rad >= LocalAABB.Min.Y && Location.Y + Rad <= LocalAABB.Max.Y))
+		{
+			OutInstances.Add(Inst);
+		}
+	}
+}
+
+void UProceduralFoliageTile::AddInstances(const TArray<FProceduralFoliageInstance*>& NewInstances, const FTransform& RelativeTM, const FBox2D& InnerLocalAABB)
+{
+	for (const FProceduralFoliageInstance* Inst : NewInstances)
+	{
+		const FVector& Location = Inst->Location;	//we need the local space because we're comparing it to the AABB
+		//Instances in InnerLocalAABB or on the border of the max sides of the AABB will be visible and instantiated by this tile
+		//Instances outside of the InnerLocalAABB are only used for rejection purposes. This is needed for overlapping tiles
+		const float Radius = Inst->GetMaxRadius();
+		bool bBlocker = false;
+		if (Location.X + Radius <= InnerLocalAABB.Min.X || Location.X - Radius > InnerLocalAABB.Max.X || Location.Y + Radius <= InnerLocalAABB.Min.Y || Location.Y - Radius > InnerLocalAABB.Max.Y)
+		{
+			//Only used for blocking instances in case of overlap, a different tile will instantiate it
+			bBlocker = true;
+		}
+
+		const FVector NewLocation = RelativeTM.TransformPosition(Inst->Location);
+		if (FProceduralFoliageInstance* NewInst = NewSeed(NewLocation, Inst->Scale, Inst->Type, Inst->Age, bBlocker))
+		{
+			
+			Instances.Add(NewInst);
+		}
+	}
+
+	FlushPendingRemovals();
+}
+
+#undef LOCTEXT_NAMESPACE
\ No newline at end of file
diff --git a/Engine/Source/Runtime/Foliage/Public/FoliageType.h b/Engine/Source/Runtime/Foliage/Public/FoliageType.h
--- a/Engine/Source/Runtime/Foliage/Public/FoliageType.h
+++ b/Engine/Source/Runtime/Foliage/Public/FoliageType.h
@@ -238,62 +238,62 @@ class UFoliageType : public UObject
 	virtual void PostEditChangeProperty(struct FPropertyChangedEvent& PropertyChangedEvent) override;
 #endif
 
-	/** Ecosystem specific parameters */
+	/** Procedural specific parameters */
 
 	/** Specifies the number of seeds to populate along 10 meters. The number is implicitly squared to cover a 10m x 10m area*/
-	UPROPERTY(Category = Ecosystem, EditAnywhere, meta = (Subcategory = "Clustering", ClampMin = "0.0", UIMin = "0.0"))
+	UPROPERTY(Category = Procedural, EditAnywhere, meta = (Subcategory = "Clustering", ClampMin = "0.0", UIMin = "0.0"))
 	float InitialSeedDensity;
 
 	/** The average distance between the spreading instance and its seeds. For example, a tree with an AverageSpreadDistance 10 will ensure the average distance between the tree and its seeds is 10cm */
-	UPROPERTY(Category = Ecosystem, EditAnywhere, meta = (Subcategory = "Clustering", ClampMin = "0.0", UIMin = "0.0"))
+	UPROPERTY(Category = Procedural, EditAnywhere, meta = (Subcategory = "Clustering", ClampMin = "0.0", UIMin = "0.0"))
 	float AverageSpreadDistance;
 
 	/** Specifies how much seed distance varies from the average. For example, a tree with an AverageSpreadDistance 10 and a SpreadVariance 1 will produce seeds with an average distance of 10cm plus or minus 1cm */
-	UPROPERTY(Category = Ecosystem, EditAnywhere, meta = (Subcategory = "Clustering", ClampMin = "0.0", UIMin = "0.0"))
+	UPROPERTY(Category = Procedural, EditAnywhere, meta = (Subcategory = "Clustering", ClampMin = "0.0", UIMin = "0.0"))
 	float SpreadVariance;
 
 	/** The CollisionRadius determines when two instances overlap. When two instances overlap a winner will be picked based on rules and priority. */
-	UPROPERTY(Category = Ecosystem, EditAnywhere, meta = (Subcategory = "Collision", ClampMin = "0.0", UIMin = "0.0"))
+	UPROPERTY(Category = Procedural, EditAnywhere, meta = (Subcategory = "Collision", ClampMin = "0.0", UIMin = "0.0"))
 	float CollisionRadius;
 
 	/** The ShadeRadius determines when two instances overlap. If an instance can grow in the shade this radius is ignored.*/
-	UPROPERTY(Category = Ecosystem, EditAnywhere, meta = (Subcategory = "Collision", ClampMin = "0.0", UIMin = "0.0"))
+	UPROPERTY(Category = Procedural, EditAnywhere, meta = (Subcategory = "Collision", ClampMin = "0.0", UIMin = "0.0"))
 	float ShadeRadius;
 
 	/** The number of times we age the species and spread its seeds. */
-	UPROPERTY(Category = Ecosystem, EditAnywhere, meta = (Subcategory = "Clustering", ClampMin = "0", UIMin = "0"))
+	UPROPERTY(Category = Procedural, EditAnywhere, meta = (Subcategory = "Clustering", ClampMin = "0", UIMin = "0"))
 	int32 NumSteps;
 
 	/** The number of seeds an instance will spread in a single step of the simulation. */
-	UPROPERTY(Category = Ecosystem, EditAnywhere, meta = (Subcategory = "Clustering", ClampMin = "0", UIMin = "0"))
+	UPROPERTY(Category = Procedural, EditAnywhere, meta = (Subcategory = "Clustering", ClampMin = "0", UIMin = "0"))
 	int32 SeedsPerStep;
 
 	/** Whether the species can grow in shade. If this is true shade radius is ignored during overlap tests*/
-	UPROPERTY(Category = Ecosystem, EditAnywhere, meta = (Subcategory = "Growth"))
+	UPROPERTY(Category = Procedural, EditAnywhere, meta = (Subcategory = "Growth"))
 	bool bGrowsInShade;
 
 	/** The minimum scale that an instance will be. Corresponds to age = 0*/
-	UPROPERTY(Category = Ecosystem, EditAnywhere, meta = (Subcategory = "Growth", ClampMin = "0.0", UIMin = "0.0"))
+	UPROPERTY(Category = Procedural, EditAnywhere, meta = (Subcategory = "Growth", ClampMin = "0.0", UIMin = "0.0"))
 	float MinScale;
 
 	/** The maximum scale that a seed will grow to. Corresponds to age = 1*/
-	UPROPERTY(Category = Ecosystem, EditAnywhere, meta = (Subcategory = "Clustering", ClampMin = "0.0", UIMin = "0.0"))
+	UPROPERTY(Category = Procedural, EditAnywhere, meta = (Subcategory = "Growth", ClampMin = "0.0", UIMin = "0.0"))
 	float MaxScale;
 
 	/** Specifies the oldest a new seed can be. The new seed will have an age randomly distributed in [0,InitMaxAge] */
-	UPROPERTY(Category = Ecosystem, EditAnywhere, meta = (Subcategory = "Growth", ClampMin = "0.0", UIMin = "0.0"))
+	UPROPERTY(Category = Procedural, EditAnywhere, meta = (Subcategory = "Growth", ClampMin = "0.0", UIMin = "0.0"))
 	float InitialMaxAge;
 
 	/** Specifies the oldest a seed can be. After reaching this age the instance will still spread seeds, but will not get any older*/
-	UPROPERTY(Category = Ecosystem, EditAnywhere, meta = (Subcategory = "Growth", ClampMin = "0.0", UIMin = "0.0"))
+	UPROPERTY(Category = Procedural, EditAnywhere, meta = (Subcategory = "Growth", ClampMin = "0.0", UIMin = "0.0"))
 	float MaxAge;
 
 	/** When two instances overlap we must determine which instance to remove. The instance with a lower OverlapPriority will be removed. In the case where OverlapPriority is the same regular simulation rules apply.*/
-	UPROPERTY(Category = Ecosystem, EditAnywhere, meta = (Subcategory = "Growth"))
+	UPROPERTY(Category = Procedural, EditAnywhere, meta = (Subcategory = "Growth"))
 	float OverlapPriority;
 
 	/** The set of meshes to be used for representing an instance placed in the world.*/
-	UPROPERTY(Category = Ecosystem, EditAnywhere, meta = (Subcategory = "Rendering"))
+	UPROPERTY(Category = Procedural, EditAnywhere)
 	TArray<UStaticMesh*> StaticMeshes;
 
 
@@ -310,7 +310,7 @@ class UFoliageType : public UObject
 private:
 
 	/** The curve used to interpolate the instance scale.*/
-	UPROPERTY(Category = Ecosystem, EditAnywhere, meta = (Subcategory = "Rendering"))
+	UPROPERTY(Category = Procedural, EditAnywhere, meta = (Subcategory = "Growth"))
 	FRuntimeFloatCurve ScaleCurve;
 };
 
diff --git a/Engine/Source/Runtime/Foliage/Public/ProceduralFoliage.h b/Engine/Source/Runtime/Foliage/Public/ProceduralFoliage.h
new file mode 100644
--- /dev/null
+++ b/Engine/Source/Runtime/Foliage/Public/ProceduralFoliage.h
@@ -0,0 +1,84 @@
+// Copyright 1998-2015 Epic Games, Inc. All Rights Reserved.
+
+#pragma once
+
+#include "FoliageType_InstancedStaticMesh.h"
+#include "ProceduralFoliageInstance.h"
+#include "ProceduralFoliage.generated.h"
+
+class UProceduralFoliageTile;
+
+
+USTRUCT()
+struct FProceduralFoliageTypeData
+{
+	GENERATED_USTRUCT_BODY()
+
+	UPROPERTY(Category = ProceduralFoliageSimulation, EditAnywhere, BlueprintReadOnly)
+	TSubclassOf<UFoliageType_InstancedStaticMesh> Type;
+
+	UPROPERTY(transient)
+	UFoliageType_InstancedStaticMesh* TypeInstance;
+
+	UPROPERTY()
+	int32 ChangeCount;
+};
+
+UCLASS(BlueprintType, Blueprintable)
+class FOLIAGE_API UProceduralFoliage : public UObject
+{
+	GENERATED_UCLASS_BODY()
+
+	/**The seed used for generating the randomness of the simulation. */
+	UPROPERTY(Category = ProceduralFoliageSimulation, EditAnywhere, BlueprintReadOnly)
+	int32 RandomSeed;
+
+	/**Cm size of the tile along one axis. The tile is square so the total area is TileSize*TileSize */
+	UPROPERTY(Category = ProceduralFoliageSimulation, EditAnywhere, BlueprintReadOnly)
+	float TileSize;
+
+	UPROPERTY(Category = ProceduralFoliageSimulation, EditAnywhere, BlueprintReadOnly)
+	int32 NumUniqueTiles;
+private:
+	UPROPERTY(Category = ProceduralFoliageSimulation, EditAnywhere)
+	TArray<FProceduralFoliageTypeData> Types;
+
+	UPROPERTY()
+	bool bNeedsSimulation;
+
+public:
+	UFUNCTION(BlueprintCallable, Category = ProceduralFoliageSimulation)
+	void Simulate(int32 NumSteps = -1);
+
+	int32 GetRandomNumber();
+
+	const TArray<FProceduralFoliageTypeData>& GetTypes() const { return Types; }
+
+	/** Returns the instances that need to spawn for a given min,max region */
+	void GetInstancesToSpawn(TArray<FProceduralFoliageInstance>& OutInstances, const FVector& Min = FVector(-FLT_MAX, -FLT_MAX, -FLT_MAX), const FVector& Max = FVector(FLT_MAX, FLT_MAX, FLT_MAX) ) const;
+
+#if WITH_EDITOR
+	virtual void PostEditChangeProperty(FPropertyChangedEvent& PropertyChangedEvent);
+#endif
+
+	virtual void Serialize(FArchive& Ar);
+
+	/** Takes a tile index and returns a random tile associated with that index. */
+	const UProceduralFoliageTile* GetRandomTile(int32 X, int32 Y);
+
+	/** Creates a temporary empty tile with the appropriate settings created for it. */
+	UProceduralFoliageTile* CreateTempTile();
+private:
+	void GenerateTile(const int32 NumSteps);
+
+	void CreateProceduralFoliageInstances();
+
+	bool AnyDirty() const;
+
+	void SetClean();
+private:
+	UPROPERTY()
+	TArray<UProceduralFoliageTile*> PrecomputedTiles;
+
+	FRandomStream RandomStream;
+};
diff --git a/Engine/Source/Runtime/Foliage/Public/ProceduralFoliageActor.h b/Engine/Source/Runtime/Foliage/Public/ProceduralFoliageActor.h
new file mode 100644
--- /dev/null
+++ b/Engine/Source/Runtime/Foliage/Public/ProceduralFoliageActor.h
@@ -0,0 +1,16 @@
+// Copyright 1998-2015 Epic Games, Inc. All Rights Reserved.
+
+#pragma once
+
+#include "ProceduralFoliageActor.generated.h"
+
+class UProceduralFoliageComponent;
+
+UCLASS()
+class FOLIAGE_API AProceduralFoliageActor: public AActor
+{
+	GENERATED_UCLASS_BODY()
+
+	UPROPERTY(Category = ProceduralFoliageTile, VisibleAnywhere, BlueprintReadOnly)
+	UProceduralFoliageComponent* ProceduralComponent;
+};
\ No newline at end of file
diff --git a/Engine/Source/Runtime/Foliage/Public/ProceduralFoliageComponent.h b/Engine/Source/Runtime/Foliage/Public/ProceduralFoliageComponent.h
new file mode 100644
--- /dev/null
+++ b/Engine/Source/Runtime/Foliage/Public/ProceduralFoliageComponent.h
@@ -0,0 +1,37 @@
+// Copyright 1998-2015 Epic Games, Inc. All Rights Reserved.
+
+#pragma once
+
+#include "ProceduralFoliageComponent.generated.h"
+
+class UProceduralFoliage;
+
+UCLASS(BlueprintType)
+class FOLIAGE_API UProceduralFoliageComponent : public USceneComponent
+{
+	GENERATED_UCLASS_BODY()
+
+	/** The number of ProceduralFoliage tiles to place along the x-axis */
+	UPROPERTY(Category = "ProceduralFoliage", BlueprintReadWrite, EditAnywhere)
+	int32 TilesX;
+
+	/** The number of ProceduralFoliage tiles to place along the y-axis */
+	UPROPERTY(Category = "ProceduralFoliage", BlueprintReadWrite, EditAnywhere)
+	int32 TilesY;
+
+	/** The extent of the volume above and below. Instances will only be placed on geometry within this volume*/
+	UPROPERTY(Category = "ProceduralFoliage", BlueprintReadWrite, EditAnywhere)
+	float HalfHeight;
+
+	/** The overlap in Cms between two tiles. */
+	UPROPERTY(Category = "ProceduralFoliage", BlueprintReadWrite, EditAnywhere)
+	float Overlap;
+
+	UPROPERTY(Category = "ProceduralFoliage", BlueprintReadWrite, EditAnywhere, meta=(DisplayName="ProceduralFoliage Asset") )
+	UProceduralFoliage* ProceduralFoliage;
+
+	void SpawnProceduralContent();
+
+private:
+	void SpawnTiles();
+};
\ No newline at end of file
diff --git a/Engine/Source/Runtime/Foliage/Public/ProceduralFoliageInstance.h b/Engine/Source/Runtime/Foliage/Public/ProceduralFoliageInstance.h
new file mode 100644
--- /dev/null
+++ b/Engine/Source/Runtime/Foliage/Public/ProceduralFoliageInstance.h
@@ -0,0 +1,108 @@
+// Copyright 1998-2015 Epic Games, Inc. All Rights Reserved.
+
+#pragma once
+
+#include "ProceduralFoliageInstance.generated.h"
+
+class ALandscapeProxy;
+class ULandscapeComponent;
+class UFoliageType;
+class UProceduralFoliageTile;
+
+#if WITH_PHYSX
+namespace physx
+{
+	class PxRigidActor;
+}
+#endif
+
+UENUM(BlueprintType)
+namespace ESimulationOverlap
+{
+	enum Type
+	{
+		/*Instances overlap with collision*/
+		CollisionOverlap,
+		/*Instances overlap with shade*/
+		ShadeOverlap,
+		/*No overlap*/
+		None
+	};
+}
+
+UENUM(BlueprintType)
+namespace ESimulationQuery
+{
+	enum Type
+	{
+		/*Instances overlap with collision*/
+		CollisionOverlap = 1,
+		/*Instances overlap with shade*/
+		ShadeOverlap = 2,
+		/*any overlap*/
+		AnyOverlap = 3
+	};
+}
+
+struct FProceduralFoliageInstance;
+struct FProceduralFoliageOverlap
+{
+	FProceduralFoliageOverlap(FProceduralFoliageInstance* InA, FProceduralFoliageInstance* InB, ESimulationOverlap::Type InType)
+	: A(InA)
+	, B(InB)
+	, OverlapType(InType)
+	{}
+
+	FProceduralFoliageInstance* A;
+	FProceduralFoliageInstance* B;
+	ESimulationOverlap::Type OverlapType;
+};
+
+USTRUCT(BlueprintType)
+struct FOLIAGE_API FProceduralFoliageInstance
+{
+public:
+	GENERATED_USTRUCT_BODY()
+	FProceduralFoliageInstance();
+	FProceduralFoliageInstance(const FProceduralFoliageInstance& Other);
+
+	static FProceduralFoliageInstance* Domination(FProceduralFoliageInstance* A, FProceduralFoliageInstance* B, ESimulationOverlap::Type OverlapType);
+
+	float GetMaxRadius() const;
+	float GetShadeRadius() const;
+	float GetCollisionRadius() const;
+
+	bool IsAlive() const { return bAlive; }
+
+	void TerminateInstance();
+
+	bool GrowsInShade() const;
+
+public:
+	UPROPERTY(Category = ProceduralFoliageInstance, EditAnywhere, BlueprintReadWrite)
+	FVector Location;
+
+	UPROPERTY()
+	FQuat Rotation;
+
+	UPROPERTY(Category = ProceduralFoliageInstance, EditAnywhere, BlueprintReadWrite)
+	FVector Normal;
+
+	UPROPERTY(Category = ProceduralFoliageInstance, EditAnywhere, BlueprintReadWrite)
+	float Age;
+
+	UPROPERTY()
+	float Scale;
+
+	UPROPERTY()
+	const UFoliageType_InstancedStaticMesh* Type;
+
+	UPROPERTY()
+	UStaticMesh* StaticMesh;
+
+	bool bBlocker;
+
+	UActorComponent* BaseComponent;
+private:
+	bool bAlive;
+};
\ No newline at end of file
diff --git a/Engine/Source/Runtime/Foliage/Public/ProceduralFoliageTile.h b/Engine/Source/Runtime/Foliage/Public/ProceduralFoliageTile.h
new file mode 100644
--- /dev/null
+++ b/Engine/Source/Runtime/Foliage/Public/ProceduralFoliageTile.h
@@ -0,0 +1,70 @@
+// Copyright 1998-2015 Epic Games, Inc. All Rights Reserved.
+
+#pragma once
+
+#include "ProceduralFoliageBroadphase.h"
+#include "ProceduralFoliageInstance.h"
+#include "ProceduralFoliageTile.generated.h"
+
+struct FProceduralFoliageInstance;
+struct FProceduralFoliageOverlap;
+class UProceduralFoliageComponent;
+class UProceduralFoliage;
+class UFoliageType_InstancedStaticMesh;
+
+UCLASS()
+class FOLIAGE_API UProceduralFoliageTile : public UObject
+{
+	GENERATED_UCLASS_BODY()
+
+	UPROPERTY()
+	UProceduralFoliage* ProceduralFoliage;
+
+	void Simulate(UProceduralFoliage* ProceduralFoliage, const int32 MaxNumSteps = -1);
+	void RemoveInstances();
+	void RemoveInstance(FProceduralFoliageInstance* Inst);
+
+	/** Takes a local AABB and finds all the instances in it. The instances are returned in the space local to the AABB. That is, an instance in the bottom left AABB will have a coordinate of (0,0) */
+	void GetInstancesInAABB(const FBox2D& LocalAABB, TArray<FProceduralFoliageInstance* >& OutInstances, bool bOnTheBorder) const;
+	void CreateInstancesToSpawn(TArray<FProceduralFoliageInstance>& OutInstances, const FTransform& WorldTM, UWorld* World, const float HalfHeight) const;
+	void AddInstances(const TArray<FProceduralFoliageInstance*>& NewInstances, const FTransform& LocalTM, const FBox2D& InnerLocalAABB);
+
+	virtual void BeginDestroy() override;
+
+	void InitSimulation(UProceduralFoliage* ProceduralFoliage);
+
+	void InstancesToArray();
+private:
+	
+	void StepSimulation();
+	void AddRandomSeeds(TArray<FProceduralFoliageInstance*>& OutInstances);
+
+	/** Determines whether the instance will survive all the overlaps, and then kills the appropriate instances. Returns true if the instance survives */
+	bool HandleOverlaps(FProceduralFoliageInstance* Instance);
+
+	/** Attempts to create a new instance and resolves any overlaps. Returns the new instance if successful for calling code to add to Instances */
+	FProceduralFoliageInstance* NewSeed(const FVector& Location, float Scale, const UFoliageType_InstancedStaticMesh* Type, float InAge, bool bBlocker = false);
+
+	void SpreadSeeds(TArray<FProceduralFoliageInstance*>& NewSeeds);
+	void AgeSeeds();
+
+	void MarkPendingRemoval(FProceduralFoliageInstance* ToRemove);
+	void FlushPendingRemovals();
+	
+private:
+
+	TSet<FProceduralFoliageInstance*> PendingRemovals;
+	TSet<FProceduralFoliageInstance*> Instances;
+
+	UPROPERTY()
+	TArray<FProceduralFoliageInstance> InstancesArray;
+
+	int32 SimulationStep;
+	FProceduralFoliageBroadphase Broadphase;	
+
+	FRandomStream RandomStream;
+
+private:
+	float GetRandomGaussian();
+	FVector GetSeedOffset(const UFoliageType_InstancedStaticMesh* Type, float MinDistance);
+};
\ No newline at end of file
