Take meshes array out of FoliageType. Waiting for new container approach

diff --git a/Engine/Source/Runtime/Foliage/Private/ProceduralFoliageComponent.cpp b/Engine/Source/Runtime/Foliage/Private/ProceduralFoliageComponent.cpp
--- a/Engine/Source/Runtime/Foliage/Private/ProceduralFoliageComponent.cpp
+++ b/Engine/Source/Runtime/Foliage/Private/ProceduralFoliageComponent.cpp
@@ -30,36 +30,26 @@ void SpawnInstances(const TArray<FProceduralFoliageInstance>& ProceduralFoliageI
 			ULevel* OwningLevel = EcoInst.BaseComponent->GetOwner()->GetLevel();
 		
 			AInstancedFoliageActor* IFA = AInstancedFoliageActor::GetInstancedFoliageActorForLevel(OwningLevel);
-			const UFoliageType_InstancedStaticMesh* Type = EcoInst.Type;
-			FFoliageMeshInfo* MeshInfo;
-			if (UStaticMesh* StaticMesh = EcoInst.StaticMesh)
+			UFoliageType_InstancedStaticMesh* Type = EcoInst.Type;
+			if (FFoliageMeshInfo* MeshInfo = IFA->FindOrAddMesh(Type))
 			{
-				if (IFA->GetSettingsForMesh(StaticMesh) == nullptr)
-				{
-					IFA->AddMesh(StaticMesh, nullptr, Type);
-				}
+				/* Convert ProceduralFoliage instance into foliage instance*/
+				Inst.DrawScale3D = FVector(EcoInst.Scale);
+				Inst.Location = EcoInst.Location;
+				Inst.Rotation = EcoInst.Rotation.Rotator();
+				Inst.Flags |= FOLIAGE_NoRandomYaw;
 
-				UFoliageType* Settings = IFA->GetSettingsForMesh(StaticMesh, &MeshInfo);
-				if (Settings && MeshInfo)
+				if (Type->AlignToNormal)
 				{
-					/* Convert ProceduralFoliage instance into foliage instance*/
-					Inst.DrawScale3D = FVector(EcoInst.Scale);
-					Inst.Location = EcoInst.Location;
-					Inst.Rotation = EcoInst.Rotation.Rotator();
-					Inst.Flags |= FOLIAGE_NoRandomYaw;
-
-					if (Type->AlignToNormal)
-					{
-						Inst.AlignToNormal(EcoInst.Normal, Type->AlignMaxAngle);
-					}
+					Inst.AlignToNormal(EcoInst.Normal, Type->AlignMaxAngle);
+				}
 
 
-					Inst.Base = EcoInst.BaseComponent;
-					Inst.Spawner = BaseComponent;
+				Inst.Base = EcoInst.BaseComponent;
+				Inst.Spawner = BaseComponent;
 
 
-					MeshInfo->AddInstance(IFA, Settings, Inst);
-				}
+				MeshInfo->AddInstance(IFA, Type, Inst);
 			}
 		}
 	}
diff --git a/Engine/Source/Runtime/Foliage/Private/ProceduralFoliageInstance.cpp b/Engine/Source/Runtime/Foliage/Private/ProceduralFoliageInstance.cpp
--- a/Engine/Source/Runtime/Foliage/Private/ProceduralFoliageInstance.cpp
+++ b/Engine/Source/Runtime/Foliage/Private/ProceduralFoliageInstance.cpp
@@ -11,7 +11,6 @@ FProceduralFoliageInstance::FProceduralFoliageInstance()
 , Age(0)
 , Scale(1)
 , Type(nullptr)
-, StaticMesh(nullptr)
 , bBlocker(false)
 , BaseComponent(nullptr)
 , bAlive(true)
@@ -25,7 +24,6 @@ FProceduralFoliageInstance::FProceduralFoliageInstance(const FProceduralFoliageI
 , Age(Other.Age)
 , Scale(Other.Scale)
 , Type(Other.Type)
-, StaticMesh(Other.StaticMesh)
 , bBlocker(Other.bBlocker)
 , BaseComponent(Other.BaseComponent)
 , bAlive(Other.bAlive)
diff --git a/Engine/Source/Runtime/Foliage/Private/ProceduralFoliageTile.cpp b/Engine/Source/Runtime/Foliage/Private/ProceduralFoliageTile.cpp
--- a/Engine/Source/Runtime/Foliage/Private/ProceduralFoliageTile.cpp
+++ b/Engine/Source/Runtime/Foliage/Private/ProceduralFoliageTile.cpp
@@ -57,7 +57,7 @@ bool UProceduralFoliageTile::HandleOverlaps(FProceduralFoliageInstance* Instance
 	return bSurvived;
 }
 
-FProceduralFoliageInstance* UProceduralFoliageTile::NewSeed(const FVector& Location, float Scale, const UFoliageType_InstancedStaticMesh* Type, float InAge, bool bBlocker)
+FProceduralFoliageInstance* UProceduralFoliageTile::NewSeed(const FVector& Location, float Scale, UFoliageType_InstancedStaticMesh* Type, float InAge, bool bBlocker)
 {
 	const float InitRadius = Type->GetMaxRadius() * Scale;
 	{
@@ -69,13 +69,6 @@ FProceduralFoliageInstance* UProceduralFoliageTile::NewSeed(const FVector& Locat
 		NewInst->Normal = FVector(0, 0, 1);
 		NewInst->Scale = Scale;
 		NewInst->bBlocker = bBlocker;
-		
-		if (Type->StaticMeshes.Num())
-		{
-			int32 MeshIdx = RandomStream.RandRange(0, Type->StaticMeshes.Num()-1);
-			NewInst->StaticMesh = Type->StaticMeshes[MeshIdx];
-		}
-		
 
 		Broadphase.Insert(NewInst);
 		const bool bSurvived = HandleOverlaps(NewInst);
@@ -133,7 +126,7 @@ void UProceduralFoliageTile::AgeSeeds()
 	{
 		if (Instance->IsAlive())
 		{
-			const UFoliageType_InstancedStaticMesh* Type = Instance->Type;
+			UFoliageType_InstancedStaticMesh* Type = Instance->Type;
 			if (SimulationStep <= Type->NumSteps)
 			{
 				const float CurrentAge = Instance->Age;
@@ -169,7 +162,7 @@ void UProceduralFoliageTile::SpreadSeeds(TArray<FProceduralFoliageInstance*>& Ne
 			continue;
 		}
 
-		const UFoliageType_InstancedStaticMesh* Type = Inst->Type;
+		UFoliageType_InstancedStaticMesh* Type = Inst->Type;
 
 		if (SimulationStep <= Type->NumSteps)
 		{
diff --git a/Engine/Source/Runtime/Foliage/Public/FoliageType.h b/Engine/Source/Runtime/Foliage/Public/FoliageType.h
--- a/Engine/Source/Runtime/Foliage/Public/FoliageType.h
+++ b/Engine/Source/Runtime/Foliage/Public/FoliageType.h
@@ -292,11 +292,6 @@ class UFoliageType : public UObject
 	UPROPERTY(Category = Procedural, EditAnywhere, meta = (Subcategory = "Growth"))
 	float OverlapPriority;
 
-	/** The set of meshes to be used for representing an instance placed in the world.*/
-	UPROPERTY(Category = Procedural, EditAnywhere)
-	TArray<UStaticMesh*> StaticMeshes;
-
-
 	UPROPERTY()
 	int32 ChangeCount;
 
diff --git a/Engine/Source/Runtime/Foliage/Public/ProceduralFoliageInstance.h b/Engine/Source/Runtime/Foliage/Public/ProceduralFoliageInstance.h
--- a/Engine/Source/Runtime/Foliage/Public/ProceduralFoliageInstance.h
+++ b/Engine/Source/Runtime/Foliage/Public/ProceduralFoliageInstance.h
@@ -95,10 +95,7 @@ public:
 	float Scale;
 
 	UPROPERTY()
-	const UFoliageType_InstancedStaticMesh* Type;
-
-	UPROPERTY()
-	UStaticMesh* StaticMesh;
+	UFoliageType_InstancedStaticMesh* Type;
 
 	bool bBlocker;
 
diff --git a/Engine/Source/Runtime/Foliage/Public/ProceduralFoliageTile.h b/Engine/Source/Runtime/Foliage/Public/ProceduralFoliageTile.h
--- a/Engine/Source/Runtime/Foliage/Public/ProceduralFoliageTile.h
+++ b/Engine/Source/Runtime/Foliage/Public/ProceduralFoliageTile.h
@@ -43,7 +43,7 @@ private:
 	bool HandleOverlaps(FProceduralFoliageInstance* Instance);
 
 	/** Attempts to create a new instance and resolves any overlaps. Returns the new instance if successful for calling code to add to Instances */
-	FProceduralFoliageInstance* NewSeed(const FVector& Location, float Scale, const UFoliageType_InstancedStaticMesh* Type, float InAge, bool bBlocker = false);
+	FProceduralFoliageInstance* NewSeed(const FVector& Location, float Scale, UFoliageType_InstancedStaticMesh* Type, float InAge, bool bBlocker = false);
 
 	void SpreadSeeds(TArray<FProceduralFoliageInstance*>& NewSeeds);
 	void AgeSeeds();
