Fix issue where procedural foliage can't be saved because it's referencing a transient instance.

diff --git a/Engine/Source/Runtime/Foliage/Private/ProceduralFoliageComponent.cpp b/Engine/Source/Runtime/Foliage/Private/ProceduralFoliageComponent.cpp
--- a/Engine/Source/Runtime/Foliage/Private/ProceduralFoliageComponent.cpp
+++ b/Engine/Source/Runtime/Foliage/Private/ProceduralFoliageComponent.cpp
@@ -30,26 +30,36 @@ void SpawnInstances(const TArray<FProceduralFoliageInstance>& ProceduralFoliageI
 			ULevel* OwningLevel = EcoInst.BaseComponent->GetOwner()->GetLevel();
 		
 			AInstancedFoliageActor* IFA = AInstancedFoliageActor::GetInstancedFoliageActorForLevel(OwningLevel);
-			UFoliageType_InstancedStaticMesh* Type = EcoInst.Type;
-			if (FFoliageMeshInfo* MeshInfo = IFA->FindOrAddMesh(Type))
+			const UFoliageType_InstancedStaticMesh* Type = EcoInst.Type;
+			FFoliageMeshInfo* MeshInfo;
+			if (UStaticMesh* StaticMesh = Type->GetStaticMesh())
 			{
-				/* Convert ProceduralFoliage instance into foliage instance*/
-				Inst.DrawScale3D = FVector(EcoInst.Scale);
-				Inst.Location = EcoInst.Location;
-				Inst.Rotation = EcoInst.Rotation.Rotator();
-				Inst.Flags |= FOLIAGE_NoRandomYaw;
-
-				if (Type->AlignToNormal)
+				if (IFA->GetSettingsForMesh(StaticMesh) == nullptr)
 				{
-					Inst.AlignToNormal(EcoInst.Normal, Type->AlignMaxAngle);
+					IFA->AddMesh(StaticMesh, nullptr, Type);
 				}
 
+				UFoliageType* Settings = IFA->GetSettingsForMesh(StaticMesh, &MeshInfo);
+				if (Settings && MeshInfo)
+				{
+					/* Convert ProceduralFoliage instance into foliage instance*/
+					Inst.DrawScale3D = FVector(EcoInst.Scale);
+					Inst.Location = EcoInst.Location;
+					Inst.Rotation = EcoInst.Rotation.Rotator();
+					Inst.Flags |= FOLIAGE_NoRandomYaw;
 
-				Inst.Base = EcoInst.BaseComponent;
-				Inst.Spawner = BaseComponent;
+					if (Type->AlignToNormal)
+					{
+						Inst.AlignToNormal(EcoInst.Normal, Type->AlignMaxAngle);
+					}
 
 
-				MeshInfo->AddInstance(IFA, Type, Inst);
+					Inst.Base = EcoInst.BaseComponent;
+					Inst.Spawner = BaseComponent;
+
+
+					MeshInfo->AddInstance(IFA, Settings, Inst);
+				}
 			}
 		}
 	}
diff --git a/Engine/Source/Runtime/Foliage/Private/ProceduralFoliageTile.cpp b/Engine/Source/Runtime/Foliage/Private/ProceduralFoliageTile.cpp
--- a/Engine/Source/Runtime/Foliage/Private/ProceduralFoliageTile.cpp
+++ b/Engine/Source/Runtime/Foliage/Private/ProceduralFoliageTile.cpp
@@ -57,7 +57,7 @@ bool UProceduralFoliageTile::HandleOverlaps(FProceduralFoliageInstance* Instance
 	return bSurvived;
 }
 
-FProceduralFoliageInstance* UProceduralFoliageTile::NewSeed(const FVector& Location, float Scale, UFoliageType_InstancedStaticMesh* Type, float InAge, bool bBlocker)
+FProceduralFoliageInstance* UProceduralFoliageTile::NewSeed(const FVector& Location, float Scale, const UFoliageType_InstancedStaticMesh* Type, float InAge, bool bBlocker)
 {
 	const float InitRadius = Type->GetMaxRadius() * Scale;
 	{
@@ -69,6 +69,8 @@ FProceduralFoliageInstance* UProceduralFoliageTile::NewSeed(const FVector& Locat
 		NewInst->Normal = FVector(0, 0, 1);
 		NewInst->Scale = Scale;
 		NewInst->bBlocker = bBlocker;
+		
+	
 
 		Broadphase.Insert(NewInst);
 		const bool bSurvived = HandleOverlaps(NewInst);
@@ -126,7 +128,7 @@ void UProceduralFoliageTile::AgeSeeds()
 	{
 		if (Instance->IsAlive())
 		{
-			UFoliageType_InstancedStaticMesh* Type = Instance->Type;
+			const UFoliageType_InstancedStaticMesh* Type = Instance->Type;
 			if (SimulationStep <= Type->NumSteps)
 			{
 				const float CurrentAge = Instance->Age;
@@ -162,7 +164,7 @@ void UProceduralFoliageTile::SpreadSeeds(TArray<FProceduralFoliageInstance*>& Ne
 			continue;
 		}
 
-		UFoliageType_InstancedStaticMesh* Type = Inst->Type;
+		const UFoliageType_InstancedStaticMesh* Type = Inst->Type;
 
 		if (SimulationStep <= Type->NumSteps)
 		{
diff --git a/Engine/Source/Runtime/Foliage/Public/FoliageType.h b/Engine/Source/Runtime/Foliage/Public/FoliageType.h
--- a/Engine/Source/Runtime/Foliage/Public/FoliageType.h
+++ b/Engine/Source/Runtime/Foliage/Public/FoliageType.h
@@ -229,7 +229,7 @@ class UFoliageType : public UObject
 	struct FBodyInstance BodyInstance;
 
 	/* Gets the mesh associated with this FoliageType */
-	virtual UStaticMesh* GetStaticMesh() PURE_VIRTUAL(UFoliageType::GetStaticMesh, return nullptr; );
+	virtual UStaticMesh* GetStaticMesh() const PURE_VIRTUAL(UFoliageType::GetStaticMesh, return nullptr; );
 
 #if WITH_EDITOR
 	/* Lets subclasses decide if the InstancedFoliageActor should reallocate its instances if the specified property change event occurs */
diff --git a/Engine/Source/Runtime/Foliage/Public/FoliageType_InstancedStaticMesh.h b/Engine/Source/Runtime/Foliage/Public/FoliageType_InstancedStaticMesh.h
--- a/Engine/Source/Runtime/Foliage/Public/FoliageType_InstancedStaticMesh.h
+++ b/Engine/Source/Runtime/Foliage/Public/FoliageType_InstancedStaticMesh.h
@@ -13,7 +13,7 @@ class UFoliageType_InstancedStaticMesh : public UFoliageType
 	UPROPERTY(EditAnywhere, Category=FoliageType)
 	UStaticMesh* Mesh;
 
-	virtual UStaticMesh* GetStaticMesh() override
+	virtual UStaticMesh* GetStaticMesh() const override
 	{
 		return Mesh;
 	}
diff --git a/Engine/Source/Runtime/Foliage/Public/ProceduralFoliageInstance.h b/Engine/Source/Runtime/Foliage/Public/ProceduralFoliageInstance.h
--- a/Engine/Source/Runtime/Foliage/Public/ProceduralFoliageInstance.h
+++ b/Engine/Source/Runtime/Foliage/Public/ProceduralFoliageInstance.h
@@ -95,7 +95,8 @@ public:
 	float Scale;
 
 	UPROPERTY()
-	UFoliageType_InstancedStaticMesh* Type;
+	const UFoliageType_InstancedStaticMesh* Type;
+
 
 	bool bBlocker;
 
diff --git a/Engine/Source/Runtime/Foliage/Public/ProceduralFoliageTile.h b/Engine/Source/Runtime/Foliage/Public/ProceduralFoliageTile.h
--- a/Engine/Source/Runtime/Foliage/Public/ProceduralFoliageTile.h
+++ b/Engine/Source/Runtime/Foliage/Public/ProceduralFoliageTile.h
@@ -43,7 +43,7 @@ private:
 	bool HandleOverlaps(FProceduralFoliageInstance* Instance);
 
 	/** Attempts to create a new instance and resolves any overlaps. Returns the new instance if successful for calling code to add to Instances */
-	FProceduralFoliageInstance* NewSeed(const FVector& Location, float Scale, UFoliageType_InstancedStaticMesh* Type, float InAge, bool bBlocker = false);
+	FProceduralFoliageInstance* NewSeed(const FVector& Location, float Scale, const UFoliageType_InstancedStaticMesh* Type, float InAge, bool bBlocker = false);
 
 	void SpreadSeeds(TArray<FProceduralFoliageInstance*>& NewSeeds);
 	void AgeSeeds();
