Added support for querying the number of instances in an area for an HISM.

diff --git a/Engine/Source/Runtime/Engine/Classes/Components/HierarchicalInstancedStaticMeshComponent.h b/Engine/Source/Runtime/Engine/Classes/Components/HierarchicalInstancedStaticMeshComponent.h
--- a/Engine/Source/Runtime/Engine/Classes/Components/HierarchicalInstancedStaticMeshComponent.h
+++ b/Engine/Source/Runtime/Engine/Classes/Components/HierarchicalInstancedStaticMeshComponent.h
@@ -97,6 +97,10 @@ public:
 	virtual bool UpdateInstanceTransform(int32 InstanceIndex, const FTransform& NewInstanceTransform, bool bWorldSpace) override;
 	virtual void ClearInstances() override;
 
+
+	/** Get the number of instances that overlap a given sphere */
+	int32 GetOverlappingSphereCount(const FSphere& Sphere) const;
+
 	virtual bool ShouldCreatePhysicsState() const override;
 
 	void BuildTree();
diff --git a/Engine/Source/Runtime/Engine/Private/HierarchicalInstancedStaticMesh.cpp b/Engine/Source/Runtime/Engine/Private/HierarchicalInstancedStaticMesh.cpp
--- a/Engine/Source/Runtime/Engine/Private/HierarchicalInstancedStaticMesh.cpp
+++ b/Engine/Source/Runtime/Engine/Private/HierarchicalInstancedStaticMesh.cpp
@@ -1871,6 +1871,28 @@ static void GatherInstanceTransformsInArea(const UHierarchicalInstancedStaticMes
 	}
 }
 
+int32 UHierarchicalInstancedStaticMeshComponent::GetOverlappingSphereCount(const FSphere& Sphere) const
+{
+	int32 Count = 0;
+	TArray<FTransform> Transforms;
+	const FBox AABB(Sphere.Center - FVector(Sphere.W), Sphere.Center + FVector(Sphere.W));
+	GatherInstanceTransformsInArea(*this, AABB, 0, Transforms);
+	const FBoxSphereBounds Bounds = StaticMesh->GetBounds();
+
+	for (const FTransform& TM : Transforms)
+	{
+		const FVector Center = TM.GetLocation();
+		const FSphere InstanceSphere(Center, Bounds.SphereRadius);
+		
+		if (Sphere.Intersects(InstanceSphere))
+		{
+			++Count;
+		}
+	}
+
+	return Count;
+}
+
 void UHierarchicalInstancedStaticMeshComponent::GetNavigationPerInstanceTransforms(const FBox& AreaBox, TArray<FTransform>& InstanceData) const
 {
 	if (IsTreeFullyBuilt())
diff --git a/Engine/Source/Runtime/Foliage/Private/FoliageStatistics.cpp b/Engine/Source/Runtime/Foliage/Private/FoliageStatistics.cpp
new file mode 100644
--- /dev/null
+++ b/Engine/Source/Runtime/Foliage/Private/FoliageStatistics.cpp
@@ -0,0 +1,23 @@
+// Copyright 1998-2015 Epic Games, Inc. All Rights Reserved.
+
+#include "FoliagePrivate.h"
+#include "FoliageStatistics.h"
+
+//////////////////////////////////////////////////////////////////////////
+// UFoliageStatics
+
+UFoliageStatistics::UFoliageStatistics(const FObjectInitializer& ObjectInitializer)
+	: Super(ObjectInitializer)
+{
+}
+
+int32 UFoliageStatistics::FoliageOverlappingSphereCount(UObject* WorldContextObject, const UStaticMesh* Mesh, FVector CenterPosition, float Radius)
+{
+	UWorld* World = GEngine->GetWorldFromContextObject(WorldContextObject);
+	const FSphere Sphere(CenterPosition, Radius);
+
+	//@TODO: properly handle sub-levels
+	AInstancedFoliageActor* IFA = AInstancedFoliageActor::GetInstancedFoliageActorForCurrentLevel(World);
+	const UFoliageType* FoliageType = IFA->GetSettingsForMesh(Mesh);
+	return IFA->GetOverlappingSphereCount(FoliageType, Sphere);
+}
\ No newline at end of file
diff --git a/Engine/Source/Runtime/Foliage/Private/FoliageStatistics.h b/Engine/Source/Runtime/Foliage/Private/FoliageStatistics.h
new file mode 100644
--- /dev/null
+++ b/Engine/Source/Runtime/Foliage/Private/FoliageStatistics.h
@@ -0,0 +1,24 @@
+// Copyright 1998-2015 Epic Games, Inc. All Rights Reserved.
+
+#pragma once
+#include "Kismet/BlueprintFunctionLibrary.h"
+#include "FoliageStatistics.generated.h"
+
+UCLASS()
+class FOLIAGE_API UFoliageStatistics : public UBlueprintFunctionLibrary
+{
+	GENERATED_UCLASS_BODY()
+
+	/**
+	* Counts how many foliage instances overlap a given sphere
+	*
+	* @param	Mesh			The static mesh we are interested in counting
+	* @param	CenterPosition	The center position of the sphere
+	* @param	Radius			The radius of the sphere.
+	*
+	* return number of foliage instances with their mesh set to Mesh that overlap the sphere
+	*/
+	UFUNCTION(BlueprintCallable, Category = "Foliage", meta = (WorldContext = "WorldContextObject", UnsafeDuringActorConstruction = "true"))
+	static int32 FoliageOverlappingSphereCount(UObject* WorldContextObject, const UStaticMesh* StaticMesh, FVector CenterPosition, float Radius);
+};
+
diff --git a/Engine/Source/Runtime/Foliage/Private/InstancedFoliage.cpp b/Engine/Source/Runtime/Foliage/Private/InstancedFoliage.cpp
--- a/Engine/Source/Runtime/Foliage/Private/InstancedFoliage.cpp
+++ b/Engine/Source/Runtime/Foliage/Private/InstancedFoliage.cpp
@@ -1042,13 +1042,21 @@ TMap<UFoliageType*, TArray<const FFoliageInstancePlacementInfo*>> AInstancedFoli
 	return Result;
 }
 
-FFoliageMeshInfo* AInstancedFoliageActor::FindMesh(UFoliageType* InType)
+FFoliageMeshInfo* AInstancedFoliageActor::FindMesh(const UFoliageType* InType)
 {
 	TUniqueObj<FFoliageMeshInfo>* MeshInfoEntry = FoliageMeshes.Find(InType);
 	FFoliageMeshInfo* MeshInfo = MeshInfoEntry ? &MeshInfoEntry->Get() : nullptr;
 	return MeshInfo;
 }
 
+const FFoliageMeshInfo* AInstancedFoliageActor::FindMesh(const UFoliageType* InType) const
+{
+	const TUniqueObj<FFoliageMeshInfo>* MeshInfoEntry = FoliageMeshes.Find(InType);
+	const FFoliageMeshInfo* MeshInfo = MeshInfoEntry ? &MeshInfoEntry->Get() : nullptr;
+	return MeshInfo;
+}
+
+
 FFoliageMeshInfo* AInstancedFoliageActor::FindOrAddMesh(UFoliageType* InType)
 {
 	TUniqueObj<FFoliageMeshInfo>* MeshInfoEntry = FoliageMeshes.Find(InType);
@@ -1056,6 +1064,19 @@ FFoliageMeshInfo* AInstancedFoliageActor::FindOrAddMesh(UFoliageType* InType)
 	return MeshInfo;
 }
 
+int32 AInstancedFoliageActor::GetOverlappingSphereCount(const UFoliageType* FoliageType, const FSphere& Sphere) const
+{
+	if (const FFoliageMeshInfo* MeshInfo = FindMesh(FoliageType))
+	{
+		if (MeshInfo->Component)
+		{
+			return MeshInfo->Component->GetOverlappingSphereCount(Sphere);
+		}
+	}
+
+	return 0;
+}
+
 FFoliageMeshInfo* AInstancedFoliageActor::AddMesh(UStaticMesh* InMesh, UFoliageType** OutSettings, const UFoliageType_InstancedStaticMesh* DefaultSettings)
 {
 	check(GetSettingsForMesh(InMesh) == nullptr);
@@ -1159,7 +1180,7 @@ void AInstancedFoliageActor::RemoveMesh(UFoliageType* InSettings)
 	CheckSelection();
 }
 
-UFoliageType* AInstancedFoliageActor::GetSettingsForMesh(UStaticMesh* InMesh, FFoliageMeshInfo** OutMeshInfo)
+UFoliageType* AInstancedFoliageActor::GetSettingsForMesh(const UStaticMesh* InMesh, FFoliageMeshInfo** OutMeshInfo)
 {
 	UFoliageType* Type = nullptr;
 	FFoliageMeshInfo* MeshInfo = nullptr;
diff --git a/Engine/Source/Runtime/Foliage/Public/InstancedFoliageActor.h b/Engine/Source/Runtime/Foliage/Public/InstancedFoliageActor.h
--- a/Engine/Source/Runtime/Foliage/Public/InstancedFoliageActor.h
+++ b/Engine/Source/Runtime/Foliage/Public/InstancedFoliageActor.h
@@ -71,11 +71,17 @@ public:
 	FOLIAGE_API void DeleteInstancesForSpawner(UActorComponent* InComponent);
 
 	// Finds a mesh entry
-	FOLIAGE_API FFoliageMeshInfo* FindMesh(UFoliageType* InType);
+	FOLIAGE_API FFoliageMeshInfo* FindMesh(const UFoliageType* InType);
+
+	// Finds a mesh entry
+	FOLIAGE_API const FFoliageMeshInfo* FindMesh(const UFoliageType* InType) const;
 
 	// Finds a mesh entry or adds it if it doesn't already exist
 	FOLIAGE_API FFoliageMeshInfo* FindOrAddMesh(UFoliageType* InType);
 
+	// Finds the number of instances overlapping with the sphere. 
+	FOLIAGE_API int32 GetOverlappingSphereCount(const UFoliageType* FoliageType, const FSphere& Sphere) const;
+
 	// Add a new static mesh.
 	FOLIAGE_API FFoliageMeshInfo* AddMesh(UStaticMesh* InMesh, UFoliageType** OutSettings = nullptr, const UFoliageType_InstancedStaticMesh* DefaultSettings = nullptr);
 	FOLIAGE_API FFoliageMeshInfo* AddMesh(UFoliageType* InType);
@@ -84,7 +90,7 @@ public:
 	FOLIAGE_API void RemoveMesh(UFoliageType* InFoliageType);
 
 	// Performs a reverse lookup from a mesh to its settings
-	FOLIAGE_API UFoliageType* GetSettingsForMesh(UStaticMesh* InMesh, FFoliageMeshInfo** OutMeshInfo = nullptr);
+	FOLIAGE_API UFoliageType* GetSettingsForMesh(const UStaticMesh* InMesh, FFoliageMeshInfo** OutMeshInfo = nullptr);
 
 	// Select an individual instance.
 	FOLIAGE_API void SelectInstance(UInstancedStaticMeshComponent* InComponent, int32 InComponentInstanceIndex, bool bToggle);
