Fix foliage not moving with the actor it's attached to:
- register OnMoveActor delegate with engine
- Call GetComponents on the InActor instead of on the InstancedFoliageActor

diff --git a/Engine/Source/Runtime/Foliage/Private/InstancedFoliage.cpp b/Engine/Source/Runtime/Foliage/Private/InstancedFoliage.cpp
--- a/Engine/Source/Runtime/Foliage/Private/InstancedFoliage.cpp
+++ b/Engine/Source/Runtime/Foliage/Private/InstancedFoliage.cpp
@@ -1680,28 +1680,23 @@ void AInstancedFoliageActor::NotifyFoliageTypeChanged(UFoliageType* FoliageType)
 
 void AInstancedFoliageActor::OnLevelActorMoved(AActor* InActor)
 {
-	if (!GetWorld()->IsPlayInEditor())
+	if (!InActor->GetWorld()->IsPlayInEditor())
 	{
 		if(AInstancedFoliageActor* IFA = AInstancedFoliageActor::GetInstancedFoliageActorForLevel(InActor->GetLevel()))
 		{
-			TInlineComponentArray<UActorComponent*> Components;
-			GetComponents(Components);
-
-			for (int32 Idx = 0; Idx < Components.Num(); ++Idx)
+			if (IFA == this)	//IFA for each level is registered so we'll get hit multiple times for each sublevel. @TODO: Should we only register the persistant level IFA?
 			{
-				IFA->MoveInstancesForMovedComponent(Components[Idx]);
+				TInlineComponentArray<UActorComponent*> Components;
+				InActor->GetComponents(Components);
+
+				for (int32 Idx = 0; Idx < Components.Num(); ++Idx)
+				{
+					IFA->MoveInstancesForMovedComponent(Components[Idx]);
+				}
 			}
 		}
 	}
 }
-
-void AInstancedFoliageActor::SubscribeToUpdateEvents()
-{
-	if (!GetWorld()->IsPlayInEditor())
-	{
-		OnLevelActorMovedDelegateHandle   = GEngine->OnActorMoved().AddUObject(this, &AInstancedFoliageActor::OnLevelActorMoved);
-	}
-}
 #endif
 
 //
@@ -1748,6 +1743,8 @@ void AInstancedFoliageActor::PostRegisterAllComponents()
 		ULevel* Level = GetLevel();
 		Level->OnApplyLevelTransform.Remove(OnApplyLevelTransformDelegateHandle); // RegisterAllComponents can be called many times before calling UnregisterAllComponents
 		OnApplyLevelTransformDelegateHandle = Level->OnApplyLevelTransform.AddUObject(this, &AInstancedFoliageActor::ApplyLevelTransform);
+		OnLevelActorMovedDelegateHandle = GEngine->OnActorMoved().AddUObject(this, &AInstancedFoliageActor::OnLevelActorMoved);
+
 	}
 #endif// WITH_EDITOR
 }
@@ -1759,6 +1756,7 @@ void AInstancedFoliageActor::PostUnregisterAllComponents()
 	{
 		ULevel* Level = GetLevel();
 		Level->OnApplyLevelTransform.Remove(OnApplyLevelTransformDelegateHandle);
+		GEngine->OnActorMoved().Remove(OnLevelActorMovedDelegateHandle);
 	}
 #endif// WITH_EDITOR
 
diff --git a/Engine/Source/Runtime/Foliage/Public/InstancedFoliageActor.h b/Engine/Source/Runtime/Foliage/Public/InstancedFoliageActor.h
--- a/Engine/Source/Runtime/Foliage/Public/InstancedFoliageActor.h
+++ b/Engine/Source/Runtime/Foliage/Public/InstancedFoliageActor.h
@@ -57,10 +57,7 @@ public:
 
 	// Move instances based on a component that has just been moved.
 	void MoveInstancesForMovedComponent(UActorComponent* InComponent);
-
-	/** Subscribes for actors transformation events */
-	void SubscribeToUpdateEvents();
-
+	
 	// Returns a map of Static Meshes and their placed instances attached to a component.
 	FOLIAGE_API TMap<UFoliageType*, TArray<const FFoliageInstancePlacementInfo*>> GetInstancesForComponent(UActorComponent* InComponent);
 
