Fix non editor build

diff --git a/Engine/Source/Runtime/Foliage/Private/InstancedFoliage.cpp b/Engine/Source/Runtime/Foliage/Private/InstancedFoliage.cpp
--- a/Engine/Source/Runtime/Foliage/Private/InstancedFoliage.cpp
+++ b/Engine/Source/Runtime/Foliage/Private/InstancedFoliage.cpp
@@ -793,7 +793,6 @@
 	PrimaryActorTick.bCanEverTick = false;
 }
 
-#if WITH_EDITOR
 AInstancedFoliageActor* AInstancedFoliageActor::GetInstancedFoliageActorForCurrentLevel(UWorld* InWorld, bool bCreateIfNone)
 {
 	for (TActorIterator<AInstancedFoliageActor> It(InWorld); It; ++It)
@@ -839,6 +838,61 @@
 	return nullptr;
 }
 
+
+int32 AInstancedFoliageActor::GetOverlappingSphereCount(const UFoliageType* FoliageType, const FSphere& Sphere) const
+{
+	if (const FFoliageMeshInfo* MeshInfo = FindMesh(FoliageType))
+	{
+		if (MeshInfo->Component)
+		{
+			return MeshInfo->Component->GetOverlappingSphereCount(Sphere);
+		}
+	}
+
+	return 0;
+}
+
+
+UFoliageType* AInstancedFoliageActor::GetSettingsForMesh(const UStaticMesh* InMesh, FFoliageMeshInfo** OutMeshInfo)
+{
+	UFoliageType* Type = nullptr;
+	FFoliageMeshInfo* MeshInfo = nullptr;
+
+	for (auto& MeshPair : FoliageMeshes)
+	{
+		UFoliageType* Settings = MeshPair.Key;
+		if (Settings && Settings->GetStaticMesh() == InMesh)
+		{
+			Type = MeshPair.Key;
+			MeshInfo = &*MeshPair.Value;
+			break;
+		}
+	}
+
+	if (OutMeshInfo)
+	{
+		*OutMeshInfo = MeshInfo;
+	}
+	return Type;
+}
+
+
+FFoliageMeshInfo* AInstancedFoliageActor::FindMesh(const UFoliageType* InType)
+{
+	TUniqueObj<FFoliageMeshInfo>* MeshInfoEntry = FoliageMeshes.Find(InType);
+	FFoliageMeshInfo* MeshInfo = MeshInfoEntry ? &MeshInfoEntry->Get() : nullptr;
+	return MeshInfo;
+}
+
+const FFoliageMeshInfo* AInstancedFoliageActor::FindMesh(const UFoliageType* InType) const
+{
+	const TUniqueObj<FFoliageMeshInfo>* MeshInfoEntry = FoliageMeshes.Find(InType);
+	const FFoliageMeshInfo* MeshInfo = MeshInfoEntry ? &MeshInfoEntry->Get() : nullptr;
+	return MeshInfo;
+}
+
+
+#if WITH_EDITOR
 void AInstancedFoliageActor::MoveInstancesForMovedComponent(UActorComponent* InComponent)
 {
 	bool bUpdatedInstances = false;
@@ -1042,21 +1096,6 @@
 	return Result;
 }
 
-FFoliageMeshInfo* AInstancedFoliageActor::FindMesh(const UFoliageType* InType)
-{
-	TUniqueObj<FFoliageMeshInfo>* MeshInfoEntry = FoliageMeshes.Find(InType);
-	FFoliageMeshInfo* MeshInfo = MeshInfoEntry ? &MeshInfoEntry->Get() : nullptr;
-	return MeshInfo;
-}
-
-const FFoliageMeshInfo* AInstancedFoliageActor::FindMesh(const UFoliageType* InType) const
-{
-	const TUniqueObj<FFoliageMeshInfo>* MeshInfoEntry = FoliageMeshes.Find(InType);
-	const FFoliageMeshInfo* MeshInfo = MeshInfoEntry ? &MeshInfoEntry->Get() : nullptr;
-	return MeshInfo;
-}
-
-
 FFoliageMeshInfo* AInstancedFoliageActor::FindOrAddMesh(UFoliageType* InType)
 {
 	TUniqueObj<FFoliageMeshInfo>* MeshInfoEntry = FoliageMeshes.Find(InType);
@@ -1064,19 +1103,6 @@
 	return MeshInfo;
 }
 
-int32 AInstancedFoliageActor::GetOverlappingSphereCount(const UFoliageType* FoliageType, const FSphere& Sphere) const
-{
-	if (const FFoliageMeshInfo* MeshInfo = FindMesh(FoliageType))
-	{
-		if (MeshInfo->Component)
-		{
-			return MeshInfo->Component->GetOverlappingSphereCount(Sphere);
-		}
-	}
-
-	return 0;
-}
-
 FFoliageMeshInfo* AInstancedFoliageActor::AddMesh(UStaticMesh* InMesh, UFoliageType** OutSettings, const UFoliageType_InstancedStaticMesh* DefaultSettings)
 {
 	check(GetSettingsForMesh(InMesh) == nullptr);
@@ -1180,29 +1206,6 @@
 	CheckSelection();
 }
 
-UFoliageType* AInstancedFoliageActor::GetSettingsForMesh(const UStaticMesh* InMesh, FFoliageMeshInfo** OutMeshInfo)
-{
-	UFoliageType* Type = nullptr;
-	FFoliageMeshInfo* MeshInfo = nullptr;
-
-	for (auto& MeshPair : FoliageMeshes)
-	{
-		UFoliageType* Settings = MeshPair.Key;
-		if (Settings && Settings->GetStaticMesh() == InMesh)
-		{
-			Type = MeshPair.Key;
-			MeshInfo = &*MeshPair.Value;
-			break;
-		}
-	}
-
-	if (OutMeshInfo)
-	{
-		*OutMeshInfo = MeshInfo;
-	}
-	return Type;
-}
-
 void AInstancedFoliageActor::SelectInstance(UInstancedStaticMeshComponent* InComponent, int32 InInstanceIndex, bool bToggle)
 {
 	Modify();
diff --git a/Engine/Source/Runtime/Foliage/Public/InstancedFoliageActor.h b/Engine/Source/Runtime/Foliage/Public/InstancedFoliageActor.h
--- a/Engine/Source/Runtime/Foliage/Public/InstancedFoliageActor.h
+++ b/Engine/Source/Runtime/Foliage/Public/InstancedFoliageActor.h
@@ -42,6 +42,36 @@ public:
 	virtual void ApplyWorldOffset(const FVector& InOffset, bool bWorldShift) override;
 	// End AActor interface.
 
+
+	// Performs a reverse lookup from a mesh to its settings
+	FOLIAGE_API UFoliageType* GetSettingsForMesh(const UStaticMesh* InMesh, FFoliageMeshInfo** OutMeshInfo = nullptr);
+
+	// Finds the number of instances overlapping with the sphere. 
+	FOLIAGE_API int32 GetOverlappingSphereCount(const UFoliageType* FoliageType, const FSphere& Sphere) const;
+
+	// Finds a mesh entry
+	FOLIAGE_API FFoliageMeshInfo* FindMesh(const UFoliageType* InType);
+
+	// Finds a mesh entry
+	FOLIAGE_API const FFoliageMeshInfo* FindMesh(const UFoliageType* InType) const;
+
+	/**
+	* Get the instanced foliage actor for the current streaming level.
+	*
+	* @param InCreationWorldIfNone			World to create the foliage instance in
+	* @param bCreateIfNone					Create if doesnt already exist
+	* returns								pointer to foliage object instance
+	*/
+	static FOLIAGE_API AInstancedFoliageActor* GetInstancedFoliageActorForCurrentLevel(UWorld* InWorld, bool bCreateIfNone = true);
+
+
+	/**
+	* Get the instanced foliage actor for the specified streaming level.
+	* @param bCreateIfNone					Create if doesnt already exist
+	* returns								pointer to foliage object instance
+	*/
+	static FOLIAGE_API AInstancedFoliageActor* GetInstancedFoliageActorForLevel(ULevel* Level, bool bCreateIfNone = true);
+
 #if WITH_EDITOR
 	virtual void PostEditUndo() override;
 
@@ -67,18 +97,9 @@ public:
 	// Deletes the instances spawned by a component
 	FOLIAGE_API void DeleteInstancesForSpawner(UActorComponent* InComponent);
 
-	// Finds a mesh entry
-	FOLIAGE_API FFoliageMeshInfo* FindMesh(const UFoliageType* InType);
-
-	// Finds a mesh entry
-	FOLIAGE_API const FFoliageMeshInfo* FindMesh(const UFoliageType* InType) const;
-
 	// Finds a mesh entry or adds it if it doesn't already exist
 	FOLIAGE_API FFoliageMeshInfo* FindOrAddMesh(UFoliageType* InType);
 
-	// Finds the number of instances overlapping with the sphere. 
-	FOLIAGE_API int32 GetOverlappingSphereCount(const UFoliageType* FoliageType, const FSphere& Sphere) const;
-
 	// Add a new static mesh.
 	FOLIAGE_API FFoliageMeshInfo* AddMesh(UStaticMesh* InMesh, UFoliageType** OutSettings = nullptr, const UFoliageType_InstancedStaticMesh* DefaultSettings = nullptr);
 	FOLIAGE_API FFoliageMeshInfo* AddMesh(UFoliageType* InType);
@@ -86,9 +107,6 @@ public:
 	// Remove the static mesh from the mesh list, and all its instances.
 	FOLIAGE_API void RemoveMesh(UFoliageType* InFoliageType);
 
-	// Performs a reverse lookup from a mesh to its settings
-	FOLIAGE_API UFoliageType* GetSettingsForMesh(const UStaticMesh* InMesh, FFoliageMeshInfo** OutMeshInfo = nullptr);
-
 	// Select an individual instance.
 	FOLIAGE_API void SelectInstance(UInstancedStaticMeshComponent* InComponent, int32 InComponentInstanceIndex, bool bToggle);
 
@@ -104,23 +122,6 @@ public:
 	// Transforms Editor specific data which is stored in world space
 	FOLIAGE_API void ApplyLevelTransform(const FTransform& LevelTransform);
 
-	/**
-	* Get the instanced foliage actor for the current streaming level.
-	*
-	* @param InCreationWorldIfNone			World to create the foliage instance in
-	* @param bCreateIfNone					Create if doesnt already exist
-	* returns								pointer to foliage object instance
-	*/
-	static FOLIAGE_API AInstancedFoliageActor* GetInstancedFoliageActorForCurrentLevel(UWorld* InWorld, bool bCreateIfNone = true);
-
-
-	/**
-	* Get the instanced foliage actor for the specified streaming level.
-	* @param bCreateIfNone					Create if doesnt already exist
-	* returns								pointer to foliage object instance
-	*/
-	static FOLIAGE_API AInstancedFoliageActor* GetInstancedFoliageActorForLevel(ULevel* Level, bool bCreateIfNone = true);
-
 	/* Called to notify InstancedFoliageActor that a UFoliageType has been modified */
 	void NotifyFoliageTypeChanged(UFoliageType* FoliageType);
 
