Fixed crash on certain foliage editor operations

diff --git a/Engine/Source/Runtime/Foliage/Private/InstancedFoliage.cpp b/Engine/Source/Runtime/Foliage/Private/InstancedFoliage.cpp
--- a/Engine/Source/Runtime/Foliage/Private/InstancedFoliage.cpp
+++ b/Engine/Source/Runtime/Foliage/Private/InstancedFoliage.cpp
@@ -1744,8 +1744,11 @@ void AInstancedFoliageActor::PostRegisterAllComponents()
 	if (!IsTemplate())
 	{
 		ULevel* Level = GetLevel();
-		Level->OnApplyLevelTransform.Remove(OnApplyLevelTransformDelegateHandle); // RegisterAllComponents can be called many times before calling UnregisterAllComponents
+		
+		Level->OnApplyLevelTransform.Remove(OnApplyLevelTransformDelegateHandle); 
 		OnApplyLevelTransformDelegateHandle = Level->OnApplyLevelTransform.AddUObject(this, &AInstancedFoliageActor::ApplyLevelTransform);
+
+		GEngine->OnActorMoved().Remove(OnLevelActorMovedDelegateHandle);
 		OnLevelActorMovedDelegateHandle = GEngine->OnActorMoved().AddUObject(this, &AInstancedFoliageActor::OnLevelActorMoved);
 
 	}
