FoliageOverlappingSphereCount now works with sublevels

diff --git a/Engine/Source/Runtime/Foliage/Private/FoliageStatistics.cpp b/Engine/Source/Runtime/Foliage/Private/FoliageStatistics.cpp
--- a/Engine/Source/Runtime/Foliage/Private/FoliageStatistics.cpp
+++ b/Engine/Source/Runtime/Foliage/Private/FoliageStatistics.cpp
@@ -16,8 +16,14 @@ int32 UFoliageStatistics::FoliageOverlappingSphereCount(UObject* WorldContextObj
 	UWorld* World = GEngine->GetWorldFromContextObject(WorldContextObject);
 	const FSphere Sphere(CenterPosition, Radius);
 
-	//@TODO: properly handle sub-levels
-	AInstancedFoliageActor* IFA = AInstancedFoliageActor::GetInstancedFoliageActorForCurrentLevel(World);
-	const UFoliageType* FoliageType = IFA->GetSettingsForMesh(Mesh);
-	return IFA->GetOverlappingSphereCount(FoliageType, Sphere);
+	int32 Count = 0;
+	for (FConstLevelIterator LvlItr = World->GetLevelIterator(); LvlItr; ++LvlItr)
+	{
+		ULevel* Level = *LvlItr;
+		AInstancedFoliageActor* IFA = AInstancedFoliageActor::GetInstancedFoliageActorForLevel(Level);
+		const UFoliageType* FoliageType = IFA->GetSettingsForMesh(Mesh);
+		Count += IFA->GetOverlappingSphereCount(FoliageType, Sphere);
+	}
+
+	return Count;
 }
\ No newline at end of file
