Better randomization of tiles

diff --git a/Engine/Source/Runtime/Foliage/Private/ProceduralFoliage.cpp b/Engine/Source/Runtime/Foliage/Private/ProceduralFoliage.cpp
--- a/Engine/Source/Runtime/Foliage/Private/ProceduralFoliage.cpp
+++ b/Engine/Source/Runtime/Foliage/Private/ProceduralFoliage.cpp
@@ -133,8 +133,13 @@ const UProceduralFoliageTile* UProceduralFoliage::GetRandomTile(int32 X, int32 Y
 
 	if (PrecomputedTiles.Num())
 	{
-		//const int32 Idx = ((Y + X) * (Y+X)) % PrecomputedTiles.Num();
-		const int32 Idx = ((Y + X) * (Y+X)) % PrecomputedTiles.Num();
+		FRandomStream HashStream;	//using this as a hash function
+		HashStream.Initialize(X);
+		const float XRand = HashStream.FRand();
+		HashStream.Initialize(Y);
+		const float YRand = HashStream.FRand();
+		const int32 RandomNumber = (RAND_MAX * XRand / (YRand + 0.01f));
+		const int32 Idx = RandomNumber % PrecomputedTiles.Num();
 		return PrecomputedTiles[Idx];
 	}
 
