Fix procedural foliage not working with BSP

diff --git a/Engine/Source/Runtime/Foliage/Private/ProceduralFoliageComponent.cpp b/Engine/Source/Runtime/Foliage/Private/ProceduralFoliageComponent.cpp
--- a/Engine/Source/Runtime/Foliage/Private/ProceduralFoliageComponent.cpp
+++ b/Engine/Source/Runtime/Foliage/Private/ProceduralFoliageComponent.cpp
@@ -27,7 +27,7 @@ void SpawnInstances(const TArray<FProceduralFoliageInstance>& ProceduralFoliageI
 	{
 		if (EcoInst.BaseComponent)
 		{
-			ULevel* OwningLevel = EcoInst.BaseComponent->GetOwner()->GetLevel();
+			ULevel* OwningLevel = EcoInst.BaseComponent->GetComponentLevel();
 		
 			AInstancedFoliageActor* IFA = AInstancedFoliageActor::GetInstancedFoliageActorForLevel(OwningLevel);
 			const UFoliageType_InstancedStaticMesh* Type = EcoInst.Type;
