Make procedural tile generation parallelized.
- Switch procedural placement to use ThreadPool
- Split UI into simulation and placement

diff --git a/Engine/Source/Runtime/Foliage/Private/ProceduralFoliage.cpp b/Engine/Source/Runtime/Foliage/Private/ProceduralFoliage.cpp
--- a/Engine/Source/Runtime/Foliage/Private/ProceduralFoliage.cpp
+++ b/Engine/Source/Runtime/Foliage/Private/ProceduralFoliage.cpp
@@ -5,6 +5,7 @@
 #include "ProceduralFoliageTile.h"
 #include "Components/BoxComponent.h"
 #include "PhysicsPublic.h"
+#include "Async/Async.h"
 
 #define LOCTEXT_NAMESPACE "ProceduralFoliage"
 
@@ -26,7 +27,7 @@ void UProceduralFoliage::PostEditChangeProperty(FPropertyChangedEvent& PropertyC
 UProceduralFoliageTile* UProceduralFoliage::CreateTempTile()
 {
 	UProceduralFoliageTile* TmpTile = NewObject<UProceduralFoliageTile>(this);
-	TmpTile->InitSimulation(this);
+	TmpTile->InitSimulation(this, 0);
 
 	return TmpTile;
 }
@@ -36,7 +37,7 @@ void UProceduralFoliage::GenerateTile(const int32 NumSteps)
 	UProceduralFoliageTile* NewTile = NewObject<UProceduralFoliageTile>(this);
 	PrecomputedTiles.Add(NewTile);
 
-	NewTile->Simulate(this, NumSteps);
+	NewTile->Simulate(this, GetRandomNumber(), NumSteps);
 }
 
 void UProceduralFoliage::CreateProceduralFoliageInstances()
@@ -152,13 +153,31 @@ void UProceduralFoliage::Simulate(int32 NumSteps)
 	CreateProceduralFoliageInstances();
 
 	PrecomputedTiles.Empty();
+	TArray<TFuture< UProceduralFoliageTile* >> Futures;
 
 	for (int i = 0; i < NumUniqueTiles; ++i)
 	{
-		GenerateTile(NumSteps);
+		UProceduralFoliageTile* NewTile = NewObject<UProceduralFoliageTile>(this);
+		const int32 RandomNumber = GetRandomNumber();
+
+		Futures.Add(Async<UProceduralFoliageTile*>(EAsyncExecution::ThreadPool, [=]()
+		{
+			NewTile->Simulate(this, RandomNumber, NumSteps);
+			return NewTile;
+		}));
+	}
+
+	FScopedSlowTask SlowTask(NumUniqueTiles, LOCTEXT("SimulateProceduralFoliage", "Simulate ProceduralFoliage..."));
+	SlowTask.MakeDialog();
+
+	for (int32 FutureIdx = 0; FutureIdx < Futures.Num(); ++FutureIdx)
+	{
+		PrecomputedTiles.Add(Futures[FutureIdx].Get());
+		SlowTask.EnterProgressFrame(1);
 	}
 
 	SetClean();
+	
 }
 
 int32 UProceduralFoliage::GetRandomNumber()
diff --git a/Engine/Source/Runtime/Foliage/Private/ProceduralFoliageComponent.cpp b/Engine/Source/Runtime/Foliage/Private/ProceduralFoliageComponent.cpp
--- a/Engine/Source/Runtime/Foliage/Private/ProceduralFoliageComponent.cpp
+++ b/Engine/Source/Runtime/Foliage/Private/ProceduralFoliageComponent.cpp
@@ -9,6 +9,8 @@
 #include "ProceduralFoliage.h"
 #include "Async/Async.h"
 
+#define LOCTEXT_NAMESPACE "ProceduralFoliage"
+
 UProceduralFoliageComponent::UProceduralFoliageComponent(const FObjectInitializer& ObjectInitializer)
 : Super(ObjectInitializer)
 {
@@ -133,6 +135,8 @@ void UProceduralFoliageComponent::SpawnTiles()
 		const FVector2D OverlapY(0.f, Overlap);
 
 		TArray<TFuture< TArray<FProceduralFoliageInstance>* >> Futures;
+		FScopedSlowTask SlowTask(TilesX * TilesY, LOCTEXT("PlaceProceduralFoliage", "Placing ProceduralFoliage..."));
+		SlowTask.MakeDialog();
 
 		for (int32 X = 0; X < TilesX; ++X)
 		{
@@ -147,7 +151,7 @@ void UProceduralFoliageComponent::SpawnTiles()
 
 				UProceduralFoliageTile* JTile = ProceduralFoliage->CreateTempTile();
 
-				Futures.Add(Async<TArray<FProceduralFoliageInstance>*>(EAsyncExecution::TaskGraph, [=]()
+				Futures.Add(Async<TArray<FProceduralFoliageInstance>*>(EAsyncExecution::ThreadPool, [=]()
 				{
 					FTransform TileTM = ComponentToWorld;
 					const FVector OrientedOffset = ComponentToWorld.TransformVectorNoScale(FVector(X, Y, 0.f) * FVector(InnerTileSize, InnerTileSize, 0.f));
@@ -201,6 +205,7 @@ void UProceduralFoliageComponent::SpawnTiles()
 				TArray<FProceduralFoliageInstance>* ProceduralFoliageInstances = Futures[FutureIdx++].Get();
 				SpawnInstances(*ProceduralFoliageInstances, World, this);
 				delete ProceduralFoliageInstances;
+				SlowTask.EnterProgressFrame(1);
 			}
 		}
 	}
@@ -222,4 +227,6 @@ void UProceduralFoliageComponent::SpawnProceduralContent()
 	
 	SpawnTiles();
 #endif
-}
\ No newline at end of file
+}
+
+#undef LOCTEXT_NAMESPACE
\ No newline at end of file
diff --git a/Engine/Source/Runtime/Foliage/Private/ProceduralFoliageTile.cpp b/Engine/Source/Runtime/Foliage/Private/ProceduralFoliageTile.cpp
--- a/Engine/Source/Runtime/Foliage/Private/ProceduralFoliageTile.cpp
+++ b/Engine/Source/Runtime/Foliage/Private/ProceduralFoliageTile.cpp
@@ -270,11 +270,9 @@ void UProceduralFoliageTile::FlushPendingRemovals()
 	PendingRemovals.Empty();
 }
 
-void UProceduralFoliageTile::InitSimulation(UProceduralFoliage* InProceduralFoliage)
+void UProceduralFoliageTile::InitSimulation(const UProceduralFoliage* InProceduralFoliage, const int32 RandomSeed)
 {
-	
-
-	RandomStream.Initialize(InProceduralFoliage->GetRandomNumber());
+	RandomStream.Initialize(RandomSeed);
 	ProceduralFoliage = InProceduralFoliage;
 	SimulationStep = 0;
 	Broadphase = FProceduralFoliageBroadphase(ProceduralFoliage->TileSize);
@@ -301,9 +299,9 @@ void UProceduralFoliageTile::StepSimulation()
 	FlushPendingRemovals();
 }
 
-void UProceduralFoliageTile::Simulate(UProceduralFoliage* InProceduralFoliage, const int32 MaxNumSteps)
+void UProceduralFoliageTile::Simulate(const UProceduralFoliage* InProceduralFoliage, const int32 RandomSeed, const int32 MaxNumSteps)
 {
-	InitSimulation(InProceduralFoliage);
+	InitSimulation(InProceduralFoliage, RandomSeed);
 
 	int32 MaxSteps = 0;
 	for(const FProceduralFoliageTypeData& Data : ProceduralFoliage->GetTypes())
@@ -320,14 +318,10 @@ void UProceduralFoliageTile::Simulate(UProceduralFoliage* InProceduralFoliage, c
 		MaxSteps = FMath::Min(MaxSteps, MaxNumSteps);	//only take as many steps as given
 	}
 
-	FScopedSlowTask SlowTask(MaxSteps, LOCTEXT("SimulateProceduralFoliage", "Simulate ProceduralFoliage..."));
-	SlowTask.MakeDialog();
-	
 	for (int32 Step = 0; Step < MaxSteps; ++Step)
 	{
 		StepSimulation();
 		++SimulationStep;
-		SlowTask.EnterProgressFrame(1);
 	}
 
 	InstancesToArray();
diff --git a/Engine/Source/Runtime/Foliage/Public/ProceduralFoliageTile.h b/Engine/Source/Runtime/Foliage/Public/ProceduralFoliageTile.h
--- a/Engine/Source/Runtime/Foliage/Public/ProceduralFoliageTile.h
+++ b/Engine/Source/Runtime/Foliage/Public/ProceduralFoliageTile.h
@@ -18,9 +18,9 @@ class FOLIAGE_API UProceduralFoliageTile : public UObject
 	GENERATED_UCLASS_BODY()
 
 	UPROPERTY()
-	UProceduralFoliage* ProceduralFoliage;
+	const UProceduralFoliage* ProceduralFoliage;
 
-	void Simulate(UProceduralFoliage* ProceduralFoliage, const int32 MaxNumSteps = -1);
+	void Simulate(const UProceduralFoliage* ProceduralFoliage, const int32 RandomSeed, const int32 MaxNumSteps = -1);
 	void RemoveInstances();
 	void RemoveInstance(FProceduralFoliageInstance* Inst);
 
@@ -31,7 +31,7 @@ class FOLIAGE_API UProceduralFoliageTile : public UObject
 
 	virtual void BeginDestroy() override;
 
-	void InitSimulation(UProceduralFoliage* ProceduralFoliage);
+	void InitSimulation(const UProceduralFoliage* ProceduralFoliage, const int32 RandomSeed);
 
 	void InstancesToArray();
 private:
