Improve memory performance during placement of tiles.

diff --git a/Engine/Source/Runtime/Foliage/Private/ProceduralFoliageBroadphase.cpp b/Engine/Source/Runtime/Foliage/Private/ProceduralFoliageBroadphase.cpp
--- a/Engine/Source/Runtime/Foliage/Private/ProceduralFoliageBroadphase.cpp
+++ b/Engine/Source/Runtime/Foliage/Private/ProceduralFoliageBroadphase.cpp
@@ -15,6 +15,11 @@ FProceduralFoliageBroadphase::FProceduralFoliageBroadphase(const FProceduralFoli
 	OtherBroadphase.QuadTree.Duplicate(&QuadTree);
 }
 
+void FProceduralFoliageBroadphase::Empty()
+{
+	QuadTree.Empty();
+}
+
 /*Takes the instance and returns the AABB that contains both its shade and collision radius*/
 FBox2D GetMaxAABB(FProceduralFoliageInstance* Instance)
 {
diff --git a/Engine/Source/Runtime/Foliage/Private/ProceduralFoliageBroadphase.h b/Engine/Source/Runtime/Foliage/Private/ProceduralFoliageBroadphase.h
--- a/Engine/Source/Runtime/Foliage/Private/ProceduralFoliageBroadphase.h
+++ b/Engine/Source/Runtime/Foliage/Private/ProceduralFoliageBroadphase.h
@@ -14,6 +14,7 @@ public:
 	bool GetOverlaps(FProceduralFoliageInstance* Instance, TArray<FProceduralFoliageOverlap>& Overlaps) const;
 	void GetInstancesInBox(const FBox2D& Box, TArray<FProceduralFoliageInstance*>& Instances) const;
 	void Remove(FProceduralFoliageInstance* Instance);
+	void Empty();
 
 private:
 	TQuadTree<FProceduralFoliageInstance*, 4> QuadTree;
diff --git a/Engine/Source/Runtime/Foliage/Private/ProceduralFoliageComponent.cpp b/Engine/Source/Runtime/Foliage/Private/ProceduralFoliageComponent.cpp
--- a/Engine/Source/Runtime/Foliage/Private/ProceduralFoliageComponent.cpp
+++ b/Engine/Source/Runtime/Foliage/Private/ProceduralFoliageComponent.cpp
@@ -189,6 +189,7 @@ void UProceduralFoliageComponent::SpawnTiles()
 					TArray<FProceduralFoliageInstance>* ProceduralFoliageInstances = new TArray<FProceduralFoliageInstance>();
 					JTile->InstancesToArray();
 					JTile->CreateInstancesToSpawn(*ProceduralFoliageInstances, TileTM, World, HalfHeight);
+					JTile->Empty();
 
 					return ProceduralFoliageInstances;
 					
diff --git a/Engine/Source/Runtime/Foliage/Private/ProceduralFoliageTile.cpp b/Engine/Source/Runtime/Foliage/Private/ProceduralFoliageTile.cpp
--- a/Engine/Source/Runtime/Foliage/Private/ProceduralFoliageTile.cpp
+++ b/Engine/Source/Runtime/Foliage/Private/ProceduralFoliageTile.cpp
@@ -378,6 +378,33 @@ void UProceduralFoliageTile::CreateInstancesToSpawn(TArray<FProceduralFoliageIns
 	}
 }
 
+void UProceduralFoliageTile::Empty()
+{
+	Broadphase.Empty();
+	InstancesArray.Empty();
+	
+	for (FProceduralFoliageInstance* Inst : Instances)
+	{
+		delete Inst;
+	}
+
+	Instances.Empty();
+	PendingRemovals.Empty();
+}
+
+SIZE_T UProceduralFoliageTile::GetResourceSize(EResourceSizeMode::Type Mode)
+{
+	SIZE_T TotalSize = 0;
+	for (FProceduralFoliageInstance* Inst : Instances)
+	{
+		TotalSize += sizeof(FProceduralFoliageInstance);
+	}
+	
+	//@TODO: account for broadphase
+	return TotalSize;
+}
+
+
 void UProceduralFoliageTile::GetInstancesInAABB(const FBox2D& LocalAABB, TArray<FProceduralFoliageInstance*>& OutInstances, bool bOnTheBorder) const
 {
 	TArray<FProceduralFoliageInstance*> InstancesInAABB;
diff --git a/Engine/Source/Runtime/Foliage/Public/ProceduralFoliageTile.h b/Engine/Source/Runtime/Foliage/Public/ProceduralFoliageTile.h
--- a/Engine/Source/Runtime/Foliage/Public/ProceduralFoliageTile.h
+++ b/Engine/Source/Runtime/Foliage/Public/ProceduralFoliageTile.h
@@ -29,11 +29,16 @@ class FOLIAGE_API UProceduralFoliageTile : public UObject
 	void CreateInstancesToSpawn(TArray<FProceduralFoliageInstance>& OutInstances, const FTransform& WorldTM, UWorld* World, const float HalfHeight) const;
 	void AddInstances(const TArray<FProceduralFoliageInstance*>& NewInstances, const FTransform& LocalTM, const FBox2D& InnerLocalAABB);
 
+	virtual SIZE_T GetResourceSize(EResourceSizeMode::Type Mode) override;
+
 	virtual void BeginDestroy() override;
 
 	void InitSimulation(const UProceduralFoliage* ProceduralFoliage, const int32 RandomSeed);
 
 	void InstancesToArray();
+
+	/** Empty arrays and sets */
+	void Empty();
 private:
 	
 	void StepSimulation();
