Fixing serialization mismatch between cooked and editor builds.

diff --git a/Engine/Source/Runtime/Foliage/Private/InstancedFoliage.cpp b/Engine/Source/Runtime/Foliage/Private/InstancedFoliage.cpp
--- a/Engine/Source/Runtime/Foliage/Private/InstancedFoliage.cpp
+++ b/Engine/Source/Runtime/Foliage/Private/InstancedFoliage.cpp
@@ -137,7 +137,7 @@ FArchive& operator<<(FArchive& Ar, FFoliageMeshInfo& MeshInfo)
 		Ar << MeshInfo.Instances;
 	}
 
-	if (Ar.CustomVer(FFoliageCustomVersion::GUID) >= FFoliageCustomVersion::AddedFoliageTypeUpdateGuid)
+	if (!Ar.ArIsFilterEditorOnly && Ar.CustomVer(FFoliageCustomVersion::GUID) >= FFoliageCustomVersion::AddedFoliageTypeUpdateGuid)
 	{
 		Ar << MeshInfo.FoliageTypeUpdateGuid;
 	}
