Add info actor for procedural foliage. This allows us to properly serialize the spawning asset

diff --git a/Engine/Source/Runtime/Foliage/Private/InstancedFoliage.cpp b/Engine/Source/Runtime/Foliage/Private/InstancedFoliage.cpp
--- a/Engine/Source/Runtime/Foliage/Private/InstancedFoliage.cpp
+++ b/Engine/Source/Runtime/Foliage/Private/InstancedFoliage.cpp
@@ -35,6 +35,8 @@ struct FFoliageCustomVersion
 		HierarchicalISMCNonTransactional = 2,
 		// Added FoliageTypeUpdateGuid
 		AddedFoliageTypeUpdateGuid = 3,
+		// Spawning actor that allows to delete foliage
+		SpawnerActor = 4,
 		// -----<new versions can be added above this line>-------------------------------------------------
 		VersionPlusOne,
 		LatestVersion = VersionPlusOne - 1
@@ -84,6 +86,13 @@ FArchive& operator<<(FArchive& Ar, FFoliageInstance& Instance)
 	
 	Ar << Instance.ZOffset;
 
+#if WITH_EDITORONLY_DATA
+	if (!Ar.ArIsFilterEditorOnly && Ar.CustomVer(FFoliageCustomVersion::GUID) >= FFoliageCustomVersion::SpawnerActor)
+	{
+		Ar << Instance.Spawner;
+	}
+#endif
+
 	return Ar;
 }
 
@@ -961,7 +970,7 @@ void AInstancedFoliageActor::DeleteInstancesForComponent(UActorComponent* InComp
 	}
 }
 
-void AInstancedFoliageActor::DeleteInstancesForSpawner(UActorComponent* InComponent)
+void AInstancedFoliageActor::DeleteInstancesForSpawner(AActor* InSpawner)
 {
 	for (auto& MeshPair : FoliageMeshes)
 	{
@@ -969,7 +978,7 @@ void AInstancedFoliageActor::DeleteInstancesForSpawner(UActorComponent* InCompon
 		TArray<int32> InstancesToRemove;
 		for (int32 InstanceIdx = 0; InstanceIdx < MeshInfo.Instances.Num(); InstanceIdx++)
 		{
-			if (MeshInfo.Instances[InstanceIdx].Spawner == InComponent)
+			if (MeshInfo.Instances[InstanceIdx].Spawner == InSpawner)
 			{
 				InstancesToRemove.Add(InstanceIdx);
 			}
diff --git a/Engine/Source/Runtime/Foliage/Private/ProceduralFoliageComponent.cpp b/Engine/Source/Runtime/Foliage/Private/ProceduralFoliageComponent.cpp
--- a/Engine/Source/Runtime/Foliage/Private/ProceduralFoliageComponent.cpp
+++ b/Engine/Source/Runtime/Foliage/Private/ProceduralFoliageComponent.cpp
@@ -7,6 +7,7 @@
 #include "InstancedFoliage.h"
 #include "InstancedFoliageActor.h"
 #include "ProceduralFoliage.h"
+#include "ProceduralFoliageLevelInfo.h"
 #include "Async/Async.h"
 
 #define LOCTEXT_NAMESPACE "ProceduralFoliage"
@@ -18,11 +19,14 @@ UProceduralFoliageComponent::UProceduralFoliageComponent(const FObjectInitialize
 	TilesY = 1;
 	Overlap = 0.f;
 	HalfHeight = 10000.f;
+	ContentGuid = FGuid::NewGuid();
 }
 
-void SpawnInstances(const TArray<FProceduralFoliageInstance>& ProceduralFoliageInstances, UWorld* World, UActorComponent* BaseComponent)
+void UProceduralFoliageComponent::SpawnInstances(const TArray<FProceduralFoliageInstance>& ProceduralFoliageInstances)
 {
 #if WITH_EDITOR
+	UWorld* World = GetWorld();
+
 	FFoliageInstance Inst;
 	
 	for (const FProceduralFoliageInstance& EcoInst : ProceduralFoliageInstances)
@@ -57,8 +61,18 @@ void SpawnInstances(const TArray<FProceduralFoliageInstance>& ProceduralFoliageI
 
 
 					Inst.Base = EcoInst.BaseComponent;
-					Inst.Spawner = BaseComponent;
-
+					ULevel* Level = EcoInst.BaseComponent->GetComponentLevel();
+					AProceduralFoliageLevelInfo*& LvlInfo = SublevelInfos.FindOrAdd(Level);
+					if (LvlInfo == nullptr)
+					{
+						//need to spawn an info actor for the level
+						FActorSpawnParameters SpawnParams;
+						SpawnParams.OverrideLevel = Level;
+						LvlInfo = World->SpawnActor<AProceduralFoliageLevelInfo>(SpawnParams);
+						LvlInfo->ProceduralContentGuid = ContentGuid;
+						SublevelInfos.Add(Level, LvlInfo);
+					}
+					Inst.Spawner = LvlInfo;
 
 					MeshInfo->AddInstance(IFA, Settings, Inst);
 				}
@@ -204,7 +218,7 @@ void UProceduralFoliageComponent::SpawnTiles()
 			for (int Y = 0; Y < TilesY; ++Y)
 			{
 				TArray<FProceduralFoliageInstance>* ProceduralFoliageInstances = Futures[FutureIdx++].Get();
-				SpawnInstances(*ProceduralFoliageInstances, World, this);
+				SpawnInstances(*ProceduralFoliageInstances);
 				delete ProceduralFoliageInstances;
 				SlowTask.EnterProgressFrame(1);
 			}
@@ -217,12 +231,19 @@ void UProceduralFoliageComponent::SpawnProceduralContent()
 {
 #if WITH_EDITOR
 	UWorld* World = GetWorld();
-	for (ULevel* Level : World->GetLevels())
+	for (TObjectIterator<AProceduralFoliageLevelInfo> It; It; ++It)
 	{
-		if (Level)
+		if (AProceduralFoliageLevelInfo* LvlInfo = *It)
 		{
-			AInstancedFoliageActor* IFA = AInstancedFoliageActor::GetInstancedFoliageActorForLevel(Level);
-			IFA->DeleteInstancesForSpawner(this);
+			if (LvlInfo->ProceduralContentGuid == ContentGuid)
+			{
+				ULevel* Level = LvlInfo->GetLevel();
+				SublevelInfos.FindOrAdd(Level) = *It;	//update our map of sublevel infos
+
+				//clear the IFA for the sublevel
+				AInstancedFoliageActor* IFA = AInstancedFoliageActor::GetInstancedFoliageActorForLevel(Level);
+				IFA->DeleteInstancesForSpawner(LvlInfo);
+			}
 		}
 	}
 	
diff --git a/Engine/Source/Runtime/Foliage/Private/ProceduralFoliageLevelInfo.cpp b/Engine/Source/Runtime/Foliage/Private/ProceduralFoliageLevelInfo.cpp
new file mode 100644
--- /dev/null
+++ b/Engine/Source/Runtime/Foliage/Private/ProceduralFoliageLevelInfo.cpp
@@ -0,0 +1,9 @@
+// Copyright 1998-2015 Epic Games, Inc. All Rights Reserved.
+
+#include "FoliagePrivate.h"
+#include "ProceduralFoliageLevelInfo.h"
+
+AProceduralFoliageLevelInfo::AProceduralFoliageLevelInfo(const FObjectInitializer& ObjectInitializer)
+: Super(ObjectInitializer)
+{
+}
diff --git a/Engine/Source/Runtime/Foliage/Private/ProceduralFoliageLevelInfo.h b/Engine/Source/Runtime/Foliage/Private/ProceduralFoliageLevelInfo.h
new file mode 100644
--- /dev/null
+++ b/Engine/Source/Runtime/Foliage/Private/ProceduralFoliageLevelInfo.h
@@ -0,0 +1,17 @@
+// Copyright 1998-2015 Epic Games, Inc. All Rights Reserved.
+
+#pragma once
+
+#include "ProceduralFoliageLevelInfo.generated.h"
+
+/** Procedural info is used for tagging which foliage instances belong to the owning procedural foliage asset in the level */
+
+//
+UCLASS()
+class FOLIAGE_API AProceduralFoliageLevelInfo : public AActor
+{
+	GENERATED_UCLASS_BODY()
+
+	UPROPERTY()
+	FGuid ProceduralContentGuid;
+};
diff --git a/Engine/Source/Runtime/Foliage/Public/InstancedFoliage.h b/Engine/Source/Runtime/Foliage/Public/InstancedFoliage.h
--- a/Engine/Source/Runtime/Foliage/Public/InstancedFoliage.h
+++ b/Engine/Source/Runtime/Foliage/Public/InstancedFoliage.h
@@ -55,11 +55,15 @@ struct FFoliageInstance : public FFoliageInstancePlacementInfo
 {
 	UActorComponent* Base;
 
-	UActorComponent* Spawner;
+#if WITH_EDITORONLY_DATA
+	AActor* Spawner;
+#endif
 
 	FFoliageInstance()
 	: Base(NULL)
+#if WITH_EDITORONLY_DATA
 	, Spawner(NULL)
+#endif
 	{}
 
 
diff --git a/Engine/Source/Runtime/Foliage/Public/InstancedFoliageActor.h b/Engine/Source/Runtime/Foliage/Public/InstancedFoliageActor.h
--- a/Engine/Source/Runtime/Foliage/Public/InstancedFoliageActor.h
+++ b/Engine/Source/Runtime/Foliage/Public/InstancedFoliageActor.h
@@ -95,7 +95,7 @@ public:
 	FOLIAGE_API void DeleteInstancesForComponent(UActorComponent* InComponent);
 
 	// Deletes the instances spawned by a component
-	FOLIAGE_API void DeleteInstancesForSpawner(UActorComponent* InComponent);
+	FOLIAGE_API void DeleteInstancesForSpawner(AActor* InSpawner);
 
 	// Finds a mesh entry or adds it if it doesn't already exist
 	FOLIAGE_API FFoliageMeshInfo* FindOrAddMesh(UFoliageType* InType);
diff --git a/Engine/Source/Runtime/Foliage/Public/ProceduralFoliageComponent.h b/Engine/Source/Runtime/Foliage/Public/ProceduralFoliageComponent.h
--- a/Engine/Source/Runtime/Foliage/Public/ProceduralFoliageComponent.h
+++ b/Engine/Source/Runtime/Foliage/Public/ProceduralFoliageComponent.h
@@ -5,6 +5,7 @@
 #include "ProceduralFoliageComponent.generated.h"
 
 class UProceduralFoliage;
+class AProceduralFoliageLevelInfo;
 
 UCLASS(BlueprintType)
 class FOLIAGE_API UProceduralFoliageComponent : public USceneComponent
@@ -34,4 +35,10 @@ class FOLIAGE_API UProceduralFoliageComponent : public USceneComponent
 
 private:
 	void SpawnTiles();
+	void SpawnInstances(const TArray<FProceduralFoliageInstance>& ProceduralFoliageInstances);
+
+private:
+	TMap<ULevel*, AProceduralFoliageLevelInfo*> SublevelInfos;
+	UPROPERTY()
+	FGuid ContentGuid;
 };
\ No newline at end of file
