Back out changelist 2410984

diff --git a/Engine/Source/Runtime/Foliage/Private/InstancedFoliage.cpp b/Engine/Source/Runtime/Foliage/Private/InstancedFoliage.cpp
--- a/Engine/Source/Runtime/Foliage/Private/InstancedFoliage.cpp
+++ b/Engine/Source/Runtime/Foliage/Private/InstancedFoliage.cpp
@@ -35,8 +35,6 @@ struct FFoliageCustomVersion
 		HierarchicalISMCNonTransactional = 2,
 		// Added FoliageTypeUpdateGuid
 		AddedFoliageTypeUpdateGuid = 3,
-		// Spawning actor that allows to delete foliage
-		SpawnerActor = 4,
 		// -----<new versions can be added above this line>-------------------------------------------------
 		VersionPlusOne,
 		LatestVersion = VersionPlusOne - 1
@@ -86,13 +84,6 @@ FArchive& operator<<(FArchive& Ar, FFoliageInstance& Instance)
 	
 	Ar << Instance.ZOffset;
 
-#if WITH_EDITORONLY_DATA
-	if (!Ar.ArIsFilterEditorOnly && Ar.CustomVer(FFoliageCustomVersion::GUID) >= FFoliageCustomVersion::SpawnerActor)
-	{
-		Ar << Instance.Spawner;
-	}
-#endif
-
 	return Ar;
 }
 
@@ -970,7 +961,7 @@ void AInstancedFoliageActor::DeleteInstancesForComponent(UActorComponent* InComp
 	}
 }
 
-void AInstancedFoliageActor::DeleteInstancesForSpawner(AActor* InSpawner)
+void AInstancedFoliageActor::DeleteInstancesForSpawner(UActorComponent* InComponent)
 {
 	for (auto& MeshPair : FoliageMeshes)
 	{
@@ -978,7 +969,7 @@ void AInstancedFoliageActor::DeleteInstancesForSpawner(AActor* InSpawner)
 		TArray<int32> InstancesToRemove;
 		for (int32 InstanceIdx = 0; InstanceIdx < MeshInfo.Instances.Num(); InstanceIdx++)
 		{
-			if (MeshInfo.Instances[InstanceIdx].Spawner == InSpawner)
+			if (MeshInfo.Instances[InstanceIdx].Spawner == InComponent)
 			{
 				InstancesToRemove.Add(InstanceIdx);
 			}
diff --git a/Engine/Source/Runtime/Foliage/Private/ProceduralFoliageComponent.cpp b/Engine/Source/Runtime/Foliage/Private/ProceduralFoliageComponent.cpp
--- a/Engine/Source/Runtime/Foliage/Private/ProceduralFoliageComponent.cpp
+++ b/Engine/Source/Runtime/Foliage/Private/ProceduralFoliageComponent.cpp
@@ -7,7 +7,6 @@
 #include "InstancedFoliage.h"
 #include "InstancedFoliageActor.h"
 #include "ProceduralFoliage.h"
-#include "ProceduralFoliageLevelInfo.h"
 #include "Async/Async.h"
 
 #define LOCTEXT_NAMESPACE "ProceduralFoliage"
@@ -19,14 +18,11 @@ UProceduralFoliageComponent::UProceduralFoliageComponent(const FObjectInitialize
 	TilesY = 1;
 	Overlap = 0.f;
 	HalfHeight = 10000.f;
-	ContentGuid = FGuid::NewGuid();
 }
 
-void UProceduralFoliageComponent::SpawnInstances(const TArray<FProceduralFoliageInstance>& ProceduralFoliageInstances)
+void SpawnInstances(const TArray<FProceduralFoliageInstance>& ProceduralFoliageInstances, UWorld* World, UActorComponent* BaseComponent)
 {
 #if WITH_EDITOR
-	UWorld* World = GetWorld();
-
 	FFoliageInstance Inst;
 	
 	for (const FProceduralFoliageInstance& EcoInst : ProceduralFoliageInstances)
@@ -61,18 +57,8 @@ void UProceduralFoliageComponent::SpawnInstances(const TArray<FProceduralFoliage
 
 
 					Inst.Base = EcoInst.BaseComponent;
-					ULevel* Level = EcoInst.BaseComponent->GetComponentLevel();
-					AProceduralFoliageLevelInfo*& LvlInfo = SublevelInfos.FindOrAdd(Level);
-					if (LvlInfo == nullptr)
-					{
-						//need to spawn an info actor for the level
-						FActorSpawnParameters SpawnParams;
-						SpawnParams.OverrideLevel = Level;
-						LvlInfo = World->SpawnActor<AProceduralFoliageLevelInfo>(SpawnParams);
-						LvlInfo->ProceduralContentGuid = ContentGuid;
-						SublevelInfos.Add(Level, LvlInfo);
-					}
-					Inst.Spawner = LvlInfo;
+					Inst.Spawner = BaseComponent;
+
 
 					MeshInfo->AddInstance(IFA, Settings, Inst);
 				}
@@ -218,7 +204,7 @@ void UProceduralFoliageComponent::SpawnTiles()
 			for (int Y = 0; Y < TilesY; ++Y)
 			{
 				TArray<FProceduralFoliageInstance>* ProceduralFoliageInstances = Futures[FutureIdx++].Get();
-				SpawnInstances(*ProceduralFoliageInstances);
+				SpawnInstances(*ProceduralFoliageInstances, World, this);
 				delete ProceduralFoliageInstances;
 				SlowTask.EnterProgressFrame(1);
 			}
@@ -231,19 +217,12 @@ void UProceduralFoliageComponent::SpawnProceduralContent()
 {
 #if WITH_EDITOR
 	UWorld* World = GetWorld();
-	for (TObjectIterator<AProceduralFoliageLevelInfo> It; It; ++It)
+	for (ULevel* Level : World->GetLevels())
 	{
-		if (AProceduralFoliageLevelInfo* LvlInfo = *It)
+		if (Level)
 		{
-			if (LvlInfo->ProceduralContentGuid == ContentGuid)
-			{
-				ULevel* Level = LvlInfo->GetLevel();
-				SublevelInfos.FindOrAdd(Level) = *It;	//update our map of sublevel infos
-
-				//clear the IFA for the sublevel
-				AInstancedFoliageActor* IFA = AInstancedFoliageActor::GetInstancedFoliageActorForLevel(Level);
-				IFA->DeleteInstancesForSpawner(LvlInfo);
-			}
+			AInstancedFoliageActor* IFA = AInstancedFoliageActor::GetInstancedFoliageActorForLevel(Level);
+			IFA->DeleteInstancesForSpawner(this);
 		}
 	}
 	
diff --git a/Engine/Source/Runtime/Foliage/Private/ProceduralFoliageLevelInfo.cpp b/Engine/Source/Runtime/Foliage/Private/ProceduralFoliageLevelInfo.cpp
deleted file mode 100644
diff --git a/Engine/Source/Runtime/Foliage/Private/ProceduralFoliageLevelInfo.h b/Engine/Source/Runtime/Foliage/Private/ProceduralFoliageLevelInfo.h
deleted file mode 100644
diff --git a/Engine/Source/Runtime/Foliage/Public/InstancedFoliage.h b/Engine/Source/Runtime/Foliage/Public/InstancedFoliage.h
--- a/Engine/Source/Runtime/Foliage/Public/InstancedFoliage.h
+++ b/Engine/Source/Runtime/Foliage/Public/InstancedFoliage.h
@@ -55,15 +55,11 @@ struct FFoliageInstance : public FFoliageInstancePlacementInfo
 {
 	UActorComponent* Base;
 
-#if WITH_EDITORONLY_DATA
-	AActor* Spawner;
-#endif
+	UActorComponent* Spawner;
 
 	FFoliageInstance()
 	: Base(NULL)
-#if WITH_EDITORONLY_DATA
 	, Spawner(NULL)
-#endif
 	{}
 
 
diff --git a/Engine/Source/Runtime/Foliage/Public/InstancedFoliageActor.h b/Engine/Source/Runtime/Foliage/Public/InstancedFoliageActor.h
--- a/Engine/Source/Runtime/Foliage/Public/InstancedFoliageActor.h
+++ b/Engine/Source/Runtime/Foliage/Public/InstancedFoliageActor.h
@@ -95,7 +95,7 @@ public:
 	FOLIAGE_API void DeleteInstancesForComponent(UActorComponent* InComponent);
 
 	// Deletes the instances spawned by a component
-	FOLIAGE_API void DeleteInstancesForSpawner(AActor* InSpawner);
+	FOLIAGE_API void DeleteInstancesForSpawner(UActorComponent* InComponent);
 
 	// Finds a mesh entry or adds it if it doesn't already exist
 	FOLIAGE_API FFoliageMeshInfo* FindOrAddMesh(UFoliageType* InType);
diff --git a/Engine/Source/Runtime/Foliage/Public/ProceduralFoliageComponent.h b/Engine/Source/Runtime/Foliage/Public/ProceduralFoliageComponent.h
--- a/Engine/Source/Runtime/Foliage/Public/ProceduralFoliageComponent.h
+++ b/Engine/Source/Runtime/Foliage/Public/ProceduralFoliageComponent.h
@@ -5,7 +5,6 @@
 #include "ProceduralFoliageComponent.generated.h"
 
 class UProceduralFoliage;
-class AProceduralFoliageLevelInfo;
 
 UCLASS(BlueprintType)
 class FOLIAGE_API UProceduralFoliageComponent : public USceneComponent
@@ -35,10 +34,4 @@ class FOLIAGE_API UProceduralFoliageComponent : public USceneComponent
 
 private:
 	void SpawnTiles();
-	void SpawnInstances(const TArray<FProceduralFoliageInstance>& ProceduralFoliageInstances);
-
-private:
-	TMap<ULevel*, AProceduralFoliageLevelInfo*> SublevelInfos;
-	UPROPERTY()
-	FGuid ContentGuid;
 };
\ No newline at end of file
