Introduce ProceduralGuid so that we can properly serialize owning procedural component

diff --git a/Engine/Source/Runtime/Foliage/Private/InstancedFoliage.cpp b/Engine/Source/Runtime/Foliage/Private/InstancedFoliage.cpp
--- a/Engine/Source/Runtime/Foliage/Private/InstancedFoliage.cpp
+++ b/Engine/Source/Runtime/Foliage/Private/InstancedFoliage.cpp
@@ -12,6 +12,7 @@
 #include "Components/ModelComponent.h"
 #include "Components/HierarchicalInstancedStaticMeshComponent.h"
 #include "Serialization/CustomVersion.h"
+#include "ProceduralFoliageComponent.h"
 
 #define LOCTEXT_NAMESPACE "InstancedFoliage"
 
@@ -35,6 +36,8 @@ struct FFoliageCustomVersion
 		HierarchicalISMCNonTransactional = 2,
 		// Added FoliageTypeUpdateGuid
 		AddedFoliageTypeUpdateGuid = 3,
+		// Use a GUID to determine whic procedural actor spawned us
+		ProceduralGuid = 4,
 		// -----<new versions can be added above this line>-------------------------------------------------
 		VersionPlusOne,
 		LatestVersion = VersionPlusOne - 1
@@ -84,6 +87,13 @@ FArchive& operator<<(FArchive& Ar, FFoliageInstance& Instance)
 	
 	Ar << Instance.ZOffset;
 
+#if WITH_EDITORONLY_DATA
+	if (!Ar.ArIsFilterEditorOnly && Ar.CustomVer(FFoliageCustomVersion::GUID) >= FFoliageCustomVersion::ProceduralGuid)
+	{
+		Ar << Instance.ProceduralGuid;
+	}
+#endif
+
 	return Ar;
 }
 
@@ -961,15 +971,16 @@ void AInstancedFoliageActor::DeleteInstancesForComponent(UActorComponent* InComp
 	}
 }
 
-void AInstancedFoliageActor::DeleteInstancesForSpawner(UActorComponent* InComponent)
+void AInstancedFoliageActor::DeleteInstancesForProceduralFoliageComponent(const UProceduralFoliageComponent* ProceduralFoliageComponent)
 {
+	const FGuid& ProceduralGuid = ProceduralFoliageComponent->GetProceduralGuid();
 	for (auto& MeshPair : FoliageMeshes)
 	{
 		FFoliageMeshInfo& MeshInfo = *MeshPair.Value;
 		TArray<int32> InstancesToRemove;
 		for (int32 InstanceIdx = 0; InstanceIdx < MeshInfo.Instances.Num(); InstanceIdx++)
 		{
-			if (MeshInfo.Instances[InstanceIdx].Spawner == InComponent)
+			if (MeshInfo.Instances[InstanceIdx].ProceduralGuid == ProceduralGuid)
 			{
 				InstancesToRemove.Add(InstanceIdx);
 			}
diff --git a/Engine/Source/Runtime/Foliage/Private/ProceduralFoliageComponent.cpp b/Engine/Source/Runtime/Foliage/Private/ProceduralFoliageComponent.cpp
--- a/Engine/Source/Runtime/Foliage/Private/ProceduralFoliageComponent.cpp
+++ b/Engine/Source/Runtime/Foliage/Private/ProceduralFoliageComponent.cpp
@@ -2,11 +2,11 @@
 
 #include "FoliagePrivate.h"
 #include "ProceduralFoliageComponent.h"
-#include "ProceduralFoliageInstance.h"
 #include "ProceduralFoliageTile.h"
 #include "InstancedFoliage.h"
 #include "InstancedFoliageActor.h"
 #include "ProceduralFoliage.h"
+
 #include "Async/Async.h"
 
 #define LOCTEXT_NAMESPACE "ProceduralFoliage"
@@ -18,11 +18,14 @@ UProceduralFoliageComponent::UProceduralFoliageComponent(const FObjectInitialize
 	TilesY = 1;
 	Overlap = 0.f;
 	HalfHeight = 10000.f;
+	ProceduralGuid = FGuid::NewGuid();
 }
 
-void SpawnInstances(const TArray<FProceduralFoliageInstance>& ProceduralFoliageInstances, UWorld* World, UActorComponent* BaseComponent)
+void UProceduralFoliageComponent::SpawnInstances(const TArray<FProceduralFoliageInstance>& ProceduralFoliageInstances)
 {
 #if WITH_EDITOR
+	UWorld* World = GetWorld();
+
 	FFoliageInstance Inst;
 	
 	for (const FProceduralFoliageInstance& EcoInst : ProceduralFoliageInstances)
@@ -57,9 +60,8 @@ void SpawnInstances(const TArray<FProceduralFoliageInstance>& ProceduralFoliageI
 
 
 					Inst.Base = EcoInst.BaseComponent;
-					Inst.Spawner = BaseComponent;
-
-
+					Inst.ProceduralGuid = ProceduralGuid;
+					
 					MeshInfo->AddInstance(IFA, Settings, Inst);
 				}
 			}
@@ -204,7 +206,7 @@ void UProceduralFoliageComponent::SpawnTiles()
 			for (int Y = 0; Y < TilesY; ++Y)
 			{
 				TArray<FProceduralFoliageInstance>* ProceduralFoliageInstances = Futures[FutureIdx++].Get();
-				SpawnInstances(*ProceduralFoliageInstances, World, this);
+				SpawnInstances(*ProceduralFoliageInstances);
 				delete ProceduralFoliageInstances;
 				SlowTask.EnterProgressFrame(1);
 			}
@@ -222,7 +224,7 @@ void UProceduralFoliageComponent::SpawnProceduralContent()
 		if (Level)
 		{
 			AInstancedFoliageActor* IFA = AInstancedFoliageActor::GetInstancedFoliageActorForLevel(Level);
-			IFA->DeleteInstancesForSpawner(this);
+			IFA->DeleteInstancesForProceduralFoliageComponent(this);
 		}
 	}
 	
diff --git a/Engine/Source/Runtime/Foliage/Public/InstancedFoliage.h b/Engine/Source/Runtime/Foliage/Public/InstancedFoliage.h
--- a/Engine/Source/Runtime/Foliage/Public/InstancedFoliage.h
+++ b/Engine/Source/Runtime/Foliage/Public/InstancedFoliage.h
@@ -55,11 +55,12 @@ struct FFoliageInstance : public FFoliageInstancePlacementInfo
 {
 	UActorComponent* Base;
 
-	UActorComponent* Spawner;
+#if WITH_EDITORONLY_DATA
+	FGuid ProceduralGuid;
+#endif
 
 	FFoliageInstance()
 	: Base(NULL)
-	, Spawner(NULL)
 	{}
 
 
diff --git a/Engine/Source/Runtime/Foliage/Public/InstancedFoliageActor.h b/Engine/Source/Runtime/Foliage/Public/InstancedFoliageActor.h
--- a/Engine/Source/Runtime/Foliage/Public/InstancedFoliageActor.h
+++ b/Engine/Source/Runtime/Foliage/Public/InstancedFoliageActor.h
@@ -14,6 +14,7 @@ class UFoliageType;
 struct FFoliageInstancePlacementInfo;
 struct FFoliageMeshInfo;
 class UFoliageType_InstancedStaticMesh;
+class UProceduralFoliageComponent;
 
 UCLASS(notplaceable, hidecategories = Object, MinimalAPI, NotBlueprintable)
 class AInstancedFoliageActor : public AActor
@@ -94,8 +95,8 @@ public:
 	// Deletes the instances attached to a component
 	FOLIAGE_API void DeleteInstancesForComponent(UActorComponent* InComponent);
 
-	// Deletes the instances spawned by a component
-	FOLIAGE_API void DeleteInstancesForSpawner(UActorComponent* InComponent);
+	// Deletes the instances spawned by a procedural component
+	void DeleteInstancesForProceduralFoliageComponent(const UProceduralFoliageComponent* ProceduralComponent);
 
 	// Finds a mesh entry or adds it if it doesn't already exist
 	FOLIAGE_API FFoliageMeshInfo* FindOrAddMesh(UFoliageType* InType);
diff --git a/Engine/Source/Runtime/Foliage/Public/ProceduralFoliageComponent.h b/Engine/Source/Runtime/Foliage/Public/ProceduralFoliageComponent.h
--- a/Engine/Source/Runtime/Foliage/Public/ProceduralFoliageComponent.h
+++ b/Engine/Source/Runtime/Foliage/Public/ProceduralFoliageComponent.h
@@ -2,9 +2,11 @@
 
 #pragma once
 
+#include "ProceduralFoliageInstance.h"
 #include "ProceduralFoliageComponent.generated.h"
 
 class UProceduralFoliage;
+class AProceduralFoliageLevelInfo;
 
 UCLASS(BlueprintType)
 class FOLIAGE_API UProceduralFoliageComponent : public USceneComponent
@@ -31,7 +33,13 @@ class FOLIAGE_API UProceduralFoliageComponent : public USceneComponent
 	UProceduralFoliage* ProceduralFoliage;
 
 	void SpawnProceduralContent();
+	const FGuid& GetProceduralGuid() const { return ProceduralGuid;  }
 
 private:
 	void SpawnTiles();
+	void SpawnInstances(const TArray<FProceduralFoliageInstance>& ProceduralFoliageInstances);
+
+private:
+	UPROPERTY()
+	FGuid ProceduralGuid;
 };
\ No newline at end of file
