Don't recompute tile overlap unless needed. Moved Overlap into procedural foliage asset

diff --git a/Engine/Source/Editor/FoliageEdit/Private/ProceduralFoliageComponentDetails.cpp b/Engine/Source/Editor/FoliageEdit/Private/ProceduralFoliageComponentDetails.cpp
--- a/Engine/Source/Editor/FoliageEdit/Private/ProceduralFoliageComponentDetails.cpp
+++ b/Engine/Source/Editor/FoliageEdit/Private/ProceduralFoliageComponentDetails.cpp
@@ -69,7 +69,6 @@ FReply FProceduralFoliageComponentDetails::OnResimulateClicked()
 			if( !UniqueProceduralFoliages.Contains( Component->ProceduralFoliage ) )
 			{
 				UniqueProceduralFoliages.Add( Component->ProceduralFoliage );
-				Component->ProceduralFoliage->Simulate();
 			}
 
 			Component->SpawnProceduralContent();
diff --git a/Engine/Source/Editor/FoliageEdit/Private/ProceduralFoliageComponentVisualizer.cpp b/Engine/Source/Editor/FoliageEdit/Private/ProceduralFoliageComponentVisualizer.cpp
--- a/Engine/Source/Editor/FoliageEdit/Private/ProceduralFoliageComponentVisualizer.cpp
+++ b/Engine/Source/Editor/FoliageEdit/Private/ProceduralFoliageComponentVisualizer.cpp
@@ -16,9 +16,10 @@ void FProceduralFoliageComponentVisualizer::DrawVisualization(const UActorCompon
 		if (ProcComponent->ProceduralFoliage)
 		{
 			const float TileSize = ProcComponent->ProceduralFoliage->TileSize;
+			const float Overlap = ProcComponent->ProceduralFoliage->Overlap;
 			const FVector TileSizeOffset(TileSize, TileSize, 0.f);
 			const FVector TileLocation = ProcComponent->GetComponentToWorld().GetLocation();
-			const FVector OverlapV(ProcComponent->Overlap, ProcComponent->Overlap, 0.f);
+			const FVector OverlapV(Overlap, Overlap, 0.f);
 			
 			for (int32 X = 0; X < ProcComponent->TilesX; ++X)
 			{
diff --git a/Engine/Source/Runtime/Foliage/Private/ProceduralFoliageComponent.cpp b/Engine/Source/Runtime/Foliage/Private/ProceduralFoliageComponent.cpp
--- a/Engine/Source/Runtime/Foliage/Private/ProceduralFoliageComponent.cpp
+++ b/Engine/Source/Runtime/Foliage/Private/ProceduralFoliageComponent.cpp
@@ -16,7 +16,6 @@ UProceduralFoliageComponent::UProceduralFoliageComponent(const FObjectInitialize
 {
 	TilesX = 1;
 	TilesY = 1;
-	Overlap = 0.f;
 	HalfHeight = 10000.f;
 	ProceduralGuid = FGuid::NewGuid();
 }
@@ -97,38 +96,27 @@ FBox2D GetTileRegion(const int32 X, const int32 Y, const int32 CountX, const int
 	return Region;
 }
 
-FBox2D GetRightOverlap(const FBox2D& Region, const float InnerSize, const float Overlap)
-{
-	return Region + FVector2D(InnerSize + Overlap, 0.f);
-}
-
-FBox2D GetLeftOverlap(const FBox2D& Region, const float Overlap)
-{
-	return Region + FVector2D(-Overlap, 0.f);
-}
-
-FBox2D GetAboveOverlap(const FBox2D& Region, const float InnerSize, const float Overlap)
-{
-	return Region + FVector2D(0.f, InnerSize + Overlap);
-}
-
-FBox2D GetBelowOverlap(const FBox2D& Region, const float Overlap)
-{
-	return Region + FVector2D(0.f, -Overlap);
-}
-
-void UProceduralFoliageComponent::SpawnTiles()
+void UProceduralFoliageComponent::MergeTiles()
 {
 #if WITH_EDITOR
+	const int32 NumTiles = TilesX*TilesY;
+	if (ProceduralFoliage->AnyDirty() == false && NumTiles == MergedTiles.Num())
+	{
+		return;
+	}
+
 	if (ProceduralFoliage)
 	{
 		/** Constants for laying out the overlapping tile grid*/
-		const float InnerTileSize = ProceduralFoliage->TileSize;
+		const float TileSize = ProceduralFoliage->TileSize;
+		const float Overlap = ProceduralFoliage->Overlap;
+		const float InnerTileSize = TileSize;
+		
 		const FVector2D InnerTileV(InnerTileSize, InnerTileSize);
 		const FVector2D InnerTileX(InnerTileSize, 0.f);
 		const FVector2D InnerTileY(0.f, InnerTileSize);
 
-		const float OuterTileSize = ProceduralFoliage->TileSize + Overlap;
+		const float OuterTileSize = TileSize + Overlap;
 		const FVector2D OuterTileV(OuterTileSize, OuterTileSize);
 		const FVector2D OuterTileX(OuterTileSize, 0.f);
 		const FVector2D OuterTileY(0.f, OuterTileSize);
@@ -136,24 +124,20 @@ void UProceduralFoliageComponent::SpawnTiles()
 		const FVector2D OverlapX(Overlap, 0.f);
 		const FVector2D OverlapY(0.f, Overlap);
 
-		TArray<TFuture< TArray<FProceduralFoliageInstance>* >> Futures;
-		FScopedSlowTask SlowTask(TilesX * TilesY, LOCTEXT("PlaceProceduralFoliage", "Placing ProceduralFoliage..."));
-		SlowTask.MakeDialog();
-
 		for (int32 X = 0; X < TilesX; ++X)
 		{
 			for (int32 Y = 0; Y < TilesY; ++Y)
 			{
-				
 				//We have to get the tiles and create new one to build on main thread
 				const UProceduralFoliageTile* Tile = ProceduralFoliage->GetRandomTile(X, Y);
 				const UProceduralFoliageTile* RightTile = (X + 1 < TilesX) ? ProceduralFoliage->GetRandomTile(X + 1, Y) : nullptr;
-				const UProceduralFoliageTile* TopTile = (Y + 1 < TilesY) ? ProceduralFoliage->GetRandomTile(X, Y+1) : nullptr;
-				const UProceduralFoliageTile* TopRightTile = (RightTile && TopTile) ? ProceduralFoliage->GetRandomTile(X + 1, Y+1) : nullptr;
+				const UProceduralFoliageTile* TopTile = (Y + 1 < TilesY) ? ProceduralFoliage->GetRandomTile(X, Y + 1) : nullptr;
+				const UProceduralFoliageTile* TopRightTile = (RightTile && TopTile) ? ProceduralFoliage->GetRandomTile(X + 1, Y + 1) : nullptr;
 
 				UProceduralFoliageTile* JTile = ProceduralFoliage->CreateTempTile();
+				WorkTiles.Add(JTile);
 
-				Futures.Add(Async<TArray<FProceduralFoliageInstance>*>(EAsyncExecution::ThreadPool, [=]()
+				TFuture<const UProceduralFoliageTile*> Future = Async<const UProceduralFoliageTile*>(EAsyncExecution::ThreadPool, [=]()
 				{
 					FTransform TileTM = ComponentToWorld;
 					const FVector OrientedOffset = ComponentToWorld.TransformVectorNoScale(FVector(X, Y, 0.f) * FVector(InnerTileSize, InnerTileSize, 0.f));
@@ -188,13 +172,45 @@ void UProceduralFoliageComponent::SpawnTiles()
 						CopyTileInstances(TopRightTile, JTile, TopRightBox, TopRightTM, Overlap);
 					}
 
-					TArray<FProceduralFoliageInstance>* ProceduralFoliageInstances = new TArray<FProceduralFoliageInstance>();
 					JTile->InstancesToArray();
+					JTile->CleanSimulatedData();
+					return JTile;
+
+				});
+
+				MergedTiles.Add(Future.Share());
+			}
+		}
+	}
+#endif
+}
+
+void UProceduralFoliageComponent::SpawnTiles()
+{
+#if WITH_EDITOR
+	if (ProceduralFoliage)
+	{
+		TArray<TFuture< TArray<FProceduralFoliageInstance>* >> Futures;
+		FScopedSlowTask SlowTask(TilesX * TilesY, LOCTEXT("PlaceProceduralFoliage", "Placing ProceduralFoliage..."));
+		SlowTask.MakeDialog();
+
+		for (int32 X = 0; X < TilesX; ++X)
+		{
+			for (int32 Y = 0; Y < TilesY; ++Y)
+			{
+				TSharedFuture<const UProceduralFoliageTile*> FutureTile = MergedTiles[X + Y*TilesX];
+
+				Futures.Add(Async<TArray<FProceduralFoliageInstance>*>(EAsyncExecution::ThreadPool, [=]()
+				{
+					FTransform TileTM = ComponentToWorld;
+					const FVector OrientedOffset = ComponentToWorld.TransformVectorNoScale(FVector(X, Y, 0.f) * FVector(ProceduralFoliage->TileSize, ProceduralFoliage->TileSize, 0.f));
+					TileTM.AddToTranslation(OrientedOffset);
+
+					TArray<FProceduralFoliageInstance>* ProceduralFoliageInstances = new TArray<FProceduralFoliageInstance>();
+					const UProceduralFoliageTile* JTile = FutureTile.Get();
 					JTile->CreateInstancesToSpawn(*ProceduralFoliageInstances, TileTM, World, HalfHeight);
-					JTile->Empty();
 
 					return ProceduralFoliageInstances;
-					
 				})
 				);
 			}
@@ -219,6 +235,7 @@ void UProceduralFoliageComponent::SpawnProceduralContent()
 {
 #if WITH_EDITOR
 	RemoveProceduralContent();
+	MergeTiles();
 	SpawnTiles();
 #endif
 }
diff --git a/Engine/Source/Runtime/Foliage/Private/ProceduralFoliageTile.cpp b/Engine/Source/Runtime/Foliage/Private/ProceduralFoliageTile.cpp
--- a/Engine/Source/Runtime/Foliage/Private/ProceduralFoliageTile.cpp
+++ b/Engine/Source/Runtime/Foliage/Private/ProceduralFoliageTile.cpp
@@ -378,10 +378,9 @@ void UProceduralFoliageTile::CreateInstancesToSpawn(TArray<FProceduralFoliageIns
 	}
 }
 
-void UProceduralFoliageTile::Empty()
+void UProceduralFoliageTile::CleanSimulatedData()
 {
 	Broadphase.Empty();
-	InstancesArray.Empty();
 	
 	for (FProceduralFoliageInstance* Inst : Instances)
 	{
diff --git a/Engine/Source/Runtime/Foliage/Public/ProceduralFoliage.h b/Engine/Source/Runtime/Foliage/Public/ProceduralFoliage.h
--- a/Engine/Source/Runtime/Foliage/Public/ProceduralFoliage.h
+++ b/Engine/Source/Runtime/Foliage/Public/ProceduralFoliage.h
@@ -37,6 +37,10 @@ class FOLIAGE_API UProceduralFoliage : public UObject
 	UPROPERTY(Category = ProceduralFoliageSimulation, EditAnywhere, BlueprintReadOnly)
 	float TileSize;
 
+	/**Amount in Cm that the tile looks around to give a more natural looking tile seam*/
+	UPROPERTY(Category = ProceduralFoliageSimulation, EditAnywhere, BlueprintReadOnly)
+	float Overlap;
+
 	UPROPERTY(Category = ProceduralFoliageSimulation, EditAnywhere, BlueprintReadOnly)
 	int32 NumUniqueTiles;
 private:
@@ -68,13 +72,15 @@ public:
 
 	/** Creates a temporary empty tile with the appropriate settings created for it. */
 	UProceduralFoliageTile* CreateTempTile();
+
+	/** Whether simulation is needed */
+	bool AnyDirty() const;
+
 private:
 	void GenerateTile(const int32 NumSteps);
 
 	void CreateProceduralFoliageInstances();
 
-	bool AnyDirty() const;
-
 	void SetClean();
 private:
 	UPROPERTY()
diff --git a/Engine/Source/Runtime/Foliage/Public/ProceduralFoliageComponent.h b/Engine/Source/Runtime/Foliage/Public/ProceduralFoliageComponent.h
--- a/Engine/Source/Runtime/Foliage/Public/ProceduralFoliageComponent.h
+++ b/Engine/Source/Runtime/Foliage/Public/ProceduralFoliageComponent.h
@@ -2,6 +2,7 @@
 
 #pragma once
 
+#include "Async/Async.h"
 #include "ProceduralFoliageInstance.h"
 #include "ProceduralFoliageComponent.generated.h"
 
@@ -25,10 +26,6 @@ class FOLIAGE_API UProceduralFoliageComponent : public USceneComponent
 	UPROPERTY(Category = "ProceduralFoliage", BlueprintReadWrite, EditAnywhere)
 	float HalfHeight;
 
-	/** The overlap in Cms between two tiles. */
-	UPROPERTY(Category = "ProceduralFoliage", BlueprintReadWrite, EditAnywhere)
-	float Overlap;
-
 	UPROPERTY(Category = "ProceduralFoliage", BlueprintReadWrite, EditAnywhere, meta=(DisplayName="ProceduralFoliage Asset") )
 	UProceduralFoliage* ProceduralFoliage;
 
@@ -39,8 +36,14 @@ class FOLIAGE_API UProceduralFoliageComponent : public USceneComponent
 private:
 	void SpawnTiles();
 	void SpawnInstances(const TArray<FProceduralFoliageInstance>& ProceduralFoliageInstances);
+	void MergeTiles();
 
 private:
 	UPROPERTY()
 	FGuid ProceduralGuid;
+
+	UPROPERTY(transient)
+	TArray<UProceduralFoliageTile*> WorkTiles;
+
+	TArray<TSharedFuture<const UProceduralFoliageTile*> > MergedTiles;
 };
\ No newline at end of file
diff --git a/Engine/Source/Runtime/Foliage/Public/ProceduralFoliageTile.h b/Engine/Source/Runtime/Foliage/Public/ProceduralFoliageTile.h
--- a/Engine/Source/Runtime/Foliage/Public/ProceduralFoliageTile.h
+++ b/Engine/Source/Runtime/Foliage/Public/ProceduralFoliageTile.h
@@ -38,7 +38,7 @@ class FOLIAGE_API UProceduralFoliageTile : public UObject
 	void InstancesToArray();
 
 	/** Empty arrays and sets */
-	void Empty();
+	void CleanSimulatedData();
 private:
 	
 	void StepSimulation();
