Back out changelist 2412930. Waiting on mac fix for futures copy constructor

diff --git a/Engine/Source/Editor/FoliageEdit/Private/ProceduralFoliageComponentDetails.cpp b/Engine/Source/Editor/FoliageEdit/Private/ProceduralFoliageComponentDetails.cpp
--- a/Engine/Source/Editor/FoliageEdit/Private/ProceduralFoliageComponentDetails.cpp
+++ b/Engine/Source/Editor/FoliageEdit/Private/ProceduralFoliageComponentDetails.cpp
@@ -69,6 +69,7 @@ FReply FProceduralFoliageComponentDetails::OnResimulateClicked()
 			if( !UniqueProceduralFoliages.Contains( Component->ProceduralFoliage ) )
 			{
 				UniqueProceduralFoliages.Add( Component->ProceduralFoliage );
+				Component->ProceduralFoliage->Simulate();
 			}
 
 			Component->SpawnProceduralContent();
diff --git a/Engine/Source/Editor/FoliageEdit/Private/ProceduralFoliageComponentVisualizer.cpp b/Engine/Source/Editor/FoliageEdit/Private/ProceduralFoliageComponentVisualizer.cpp
--- a/Engine/Source/Editor/FoliageEdit/Private/ProceduralFoliageComponentVisualizer.cpp
+++ b/Engine/Source/Editor/FoliageEdit/Private/ProceduralFoliageComponentVisualizer.cpp
@@ -16,10 +16,9 @@ void FProceduralFoliageComponentVisualizer::DrawVisualization(const UActorCompon
 		if (ProcComponent->ProceduralFoliage)
 		{
 			const float TileSize = ProcComponent->ProceduralFoliage->TileSize;
-			const float Overlap = ProcComponent->ProceduralFoliage->Overlap;
 			const FVector TileSizeOffset(TileSize, TileSize, 0.f);
 			const FVector TileLocation = ProcComponent->GetComponentToWorld().GetLocation();
-			const FVector OverlapV(Overlap, Overlap, 0.f);
+			const FVector OverlapV(ProcComponent->Overlap, ProcComponent->Overlap, 0.f);
 			
 			for (int32 X = 0; X < ProcComponent->TilesX; ++X)
 			{
diff --git a/Engine/Source/Runtime/Foliage/Private/ProceduralFoliageComponent.cpp b/Engine/Source/Runtime/Foliage/Private/ProceduralFoliageComponent.cpp
--- a/Engine/Source/Runtime/Foliage/Private/ProceduralFoliageComponent.cpp
+++ b/Engine/Source/Runtime/Foliage/Private/ProceduralFoliageComponent.cpp
@@ -16,6 +16,7 @@ UProceduralFoliageComponent::UProceduralFoliageComponent(const FObjectInitialize
 {
 	TilesX = 1;
 	TilesY = 1;
+	Overlap = 0.f;
 	HalfHeight = 10000.f;
 	ProceduralGuid = FGuid::NewGuid();
 }
@@ -96,27 +97,38 @@ FBox2D GetTileRegion(const int32 X, const int32 Y, const int32 CountX, const int
 	return Region;
 }
 
-void UProceduralFoliageComponent::MergeTiles()
+FBox2D GetRightOverlap(const FBox2D& Region, const float InnerSize, const float Overlap)
 {
-#if WITH_EDITOR
-	const int32 NumTiles = TilesX*TilesY;
-	if (ProceduralFoliage->AnyDirty() == false && NumTiles == MergedTiles.Num())
-	{
-		return;
-	}
+	return Region + FVector2D(InnerSize + Overlap, 0.f);
+}
+
+FBox2D GetLeftOverlap(const FBox2D& Region, const float Overlap)
+{
+	return Region + FVector2D(-Overlap, 0.f);
+}
+
+FBox2D GetAboveOverlap(const FBox2D& Region, const float InnerSize, const float Overlap)
+{
+	return Region + FVector2D(0.f, InnerSize + Overlap);
+}
+
+FBox2D GetBelowOverlap(const FBox2D& Region, const float Overlap)
+{
+	return Region + FVector2D(0.f, -Overlap);
+}
 
+void UProceduralFoliageComponent::SpawnTiles()
+{
+#if WITH_EDITOR
 	if (ProceduralFoliage)
 	{
 		/** Constants for laying out the overlapping tile grid*/
-		const float TileSize = ProceduralFoliage->TileSize;
-		const float Overlap = ProceduralFoliage->Overlap;
-		const float InnerTileSize = TileSize;
-		
+		const float InnerTileSize = ProceduralFoliage->TileSize;
 		const FVector2D InnerTileV(InnerTileSize, InnerTileSize);
 		const FVector2D InnerTileX(InnerTileSize, 0.f);
 		const FVector2D InnerTileY(0.f, InnerTileSize);
 
-		const float OuterTileSize = TileSize + Overlap;
+		const float OuterTileSize = ProceduralFoliage->TileSize + Overlap;
 		const FVector2D OuterTileV(OuterTileSize, OuterTileSize);
 		const FVector2D OuterTileX(OuterTileSize, 0.f);
 		const FVector2D OuterTileY(0.f, OuterTileSize);
@@ -124,20 +136,24 @@ void UProceduralFoliageComponent::MergeTiles()
 		const FVector2D OverlapX(Overlap, 0.f);
 		const FVector2D OverlapY(0.f, Overlap);
 
+		TArray<TFuture< TArray<FProceduralFoliageInstance>* >> Futures;
+		FScopedSlowTask SlowTask(TilesX * TilesY, LOCTEXT("PlaceProceduralFoliage", "Placing ProceduralFoliage..."));
+		SlowTask.MakeDialog();
+
 		for (int32 X = 0; X < TilesX; ++X)
 		{
 			for (int32 Y = 0; Y < TilesY; ++Y)
 			{
+				
 				//We have to get the tiles and create new one to build on main thread
 				const UProceduralFoliageTile* Tile = ProceduralFoliage->GetRandomTile(X, Y);
 				const UProceduralFoliageTile* RightTile = (X + 1 < TilesX) ? ProceduralFoliage->GetRandomTile(X + 1, Y) : nullptr;
-				const UProceduralFoliageTile* TopTile = (Y + 1 < TilesY) ? ProceduralFoliage->GetRandomTile(X, Y + 1) : nullptr;
-				const UProceduralFoliageTile* TopRightTile = (RightTile && TopTile) ? ProceduralFoliage->GetRandomTile(X + 1, Y + 1) : nullptr;
+				const UProceduralFoliageTile* TopTile = (Y + 1 < TilesY) ? ProceduralFoliage->GetRandomTile(X, Y+1) : nullptr;
+				const UProceduralFoliageTile* TopRightTile = (RightTile && TopTile) ? ProceduralFoliage->GetRandomTile(X + 1, Y+1) : nullptr;
 
 				UProceduralFoliageTile* JTile = ProceduralFoliage->CreateTempTile();
-				WorkTiles.Add(JTile);
 
-				TFuture<const UProceduralFoliageTile*> Future = Async<const UProceduralFoliageTile*>(EAsyncExecution::ThreadPool, [=]()
+				Futures.Add(Async<TArray<FProceduralFoliageInstance>*>(EAsyncExecution::ThreadPool, [=]()
 				{
 					FTransform TileTM = ComponentToWorld;
 					const FVector OrientedOffset = ComponentToWorld.TransformVectorNoScale(FVector(X, Y, 0.f) * FVector(InnerTileSize, InnerTileSize, 0.f));
@@ -172,45 +188,13 @@ void UProceduralFoliageComponent::MergeTiles()
 						CopyTileInstances(TopRightTile, JTile, TopRightBox, TopRightTM, Overlap);
 					}
 
-					JTile->InstancesToArray();
-					JTile->CleanSimulatedData();
-					return JTile;
-
-				});
-
-				MergedTiles.Add(Future.Share());
-			}
-		}
-	}
-#endif
-}
-
-void UProceduralFoliageComponent::SpawnTiles()
-{
-#if WITH_EDITOR
-	if (ProceduralFoliage)
-	{
-		TArray<TFuture< TArray<FProceduralFoliageInstance>* >> Futures;
-		FScopedSlowTask SlowTask(TilesX * TilesY, LOCTEXT("PlaceProceduralFoliage", "Placing ProceduralFoliage..."));
-		SlowTask.MakeDialog();
-
-		for (int32 X = 0; X < TilesX; ++X)
-		{
-			for (int32 Y = 0; Y < TilesY; ++Y)
-			{
-				TSharedFuture<const UProceduralFoliageTile*> FutureTile = MergedTiles[X + Y*TilesX];
-
-				Futures.Add(Async<TArray<FProceduralFoliageInstance>*>(EAsyncExecution::ThreadPool, [=]()
-				{
-					FTransform TileTM = ComponentToWorld;
-					const FVector OrientedOffset = ComponentToWorld.TransformVectorNoScale(FVector(X, Y, 0.f) * FVector(ProceduralFoliage->TileSize, ProceduralFoliage->TileSize, 0.f));
-					TileTM.AddToTranslation(OrientedOffset);
-
 					TArray<FProceduralFoliageInstance>* ProceduralFoliageInstances = new TArray<FProceduralFoliageInstance>();
-					const UProceduralFoliageTile* JTile = FutureTile.Get();
+					JTile->InstancesToArray();
 					JTile->CreateInstancesToSpawn(*ProceduralFoliageInstances, TileTM, World, HalfHeight);
+					JTile->Empty();
 
 					return ProceduralFoliageInstances;
+					
 				})
 				);
 			}
@@ -235,7 +219,6 @@ void UProceduralFoliageComponent::SpawnProceduralContent()
 {
 #if WITH_EDITOR
 	RemoveProceduralContent();
-	MergeTiles();
 	SpawnTiles();
 #endif
 }
diff --git a/Engine/Source/Runtime/Foliage/Private/ProceduralFoliageTile.cpp b/Engine/Source/Runtime/Foliage/Private/ProceduralFoliageTile.cpp
--- a/Engine/Source/Runtime/Foliage/Private/ProceduralFoliageTile.cpp
+++ b/Engine/Source/Runtime/Foliage/Private/ProceduralFoliageTile.cpp
@@ -378,9 +378,10 @@ void UProceduralFoliageTile::CreateInstancesToSpawn(TArray<FProceduralFoliageIns
 	}
 }
 
-void UProceduralFoliageTile::CleanSimulatedData()
+void UProceduralFoliageTile::Empty()
 {
 	Broadphase.Empty();
+	InstancesArray.Empty();
 	
 	for (FProceduralFoliageInstance* Inst : Instances)
 	{
diff --git a/Engine/Source/Runtime/Foliage/Public/ProceduralFoliage.h b/Engine/Source/Runtime/Foliage/Public/ProceduralFoliage.h
--- a/Engine/Source/Runtime/Foliage/Public/ProceduralFoliage.h
+++ b/Engine/Source/Runtime/Foliage/Public/ProceduralFoliage.h
@@ -37,10 +37,6 @@ class FOLIAGE_API UProceduralFoliage : public UObject
 	UPROPERTY(Category = ProceduralFoliageSimulation, EditAnywhere, BlueprintReadOnly)
 	float TileSize;
 
-	/**Amount in Cm that the tile looks around to give a more natural looking tile seam*/
-	UPROPERTY(Category = ProceduralFoliageSimulation, EditAnywhere, BlueprintReadOnly)
-	float Overlap;
-
 	UPROPERTY(Category = ProceduralFoliageSimulation, EditAnywhere, BlueprintReadOnly)
 	int32 NumUniqueTiles;
 private:
@@ -72,15 +68,13 @@ public:
 
 	/** Creates a temporary empty tile with the appropriate settings created for it. */
 	UProceduralFoliageTile* CreateTempTile();
-
-	/** Whether simulation is needed */
-	bool AnyDirty() const;
-
 private:
 	void GenerateTile(const int32 NumSteps);
 
 	void CreateProceduralFoliageInstances();
 
+	bool AnyDirty() const;
+
 	void SetClean();
 private:
 	UPROPERTY()
diff --git a/Engine/Source/Runtime/Foliage/Public/ProceduralFoliageComponent.h b/Engine/Source/Runtime/Foliage/Public/ProceduralFoliageComponent.h
--- a/Engine/Source/Runtime/Foliage/Public/ProceduralFoliageComponent.h
+++ b/Engine/Source/Runtime/Foliage/Public/ProceduralFoliageComponent.h
@@ -2,7 +2,6 @@
 
 #pragma once
 
-#include "Async/Async.h"
 #include "ProceduralFoliageInstance.h"
 #include "ProceduralFoliageComponent.generated.h"
 
@@ -26,6 +25,10 @@ class FOLIAGE_API UProceduralFoliageComponent : public USceneComponent
 	UPROPERTY(Category = "ProceduralFoliage", BlueprintReadWrite, EditAnywhere)
 	float HalfHeight;
 
+	/** The overlap in Cms between two tiles. */
+	UPROPERTY(Category = "ProceduralFoliage", BlueprintReadWrite, EditAnywhere)
+	float Overlap;
+
 	UPROPERTY(Category = "ProceduralFoliage", BlueprintReadWrite, EditAnywhere, meta=(DisplayName="ProceduralFoliage Asset") )
 	UProceduralFoliage* ProceduralFoliage;
 
@@ -36,14 +39,8 @@ class FOLIAGE_API UProceduralFoliageComponent : public USceneComponent
 private:
 	void SpawnTiles();
 	void SpawnInstances(const TArray<FProceduralFoliageInstance>& ProceduralFoliageInstances);
-	void MergeTiles();
 
 private:
 	UPROPERTY()
 	FGuid ProceduralGuid;
-
-	UPROPERTY(transient)
-	TArray<UProceduralFoliageTile*> WorkTiles;
-
-	TArray<TSharedFuture<const UProceduralFoliageTile*> > MergedTiles;
 };
\ No newline at end of file
diff --git a/Engine/Source/Runtime/Foliage/Public/ProceduralFoliageTile.h b/Engine/Source/Runtime/Foliage/Public/ProceduralFoliageTile.h
--- a/Engine/Source/Runtime/Foliage/Public/ProceduralFoliageTile.h
+++ b/Engine/Source/Runtime/Foliage/Public/ProceduralFoliageTile.h
@@ -38,7 +38,7 @@ class FOLIAGE_API UProceduralFoliageTile : public UObject
 	void InstancesToArray();
 
 	/** Empty arrays and sets */
-	void CleanSimulatedData();
+	void Empty();
 private:
 	
 	void StepSimulation();
