Add min ground slope to foliage type

diff --git a/Engine/Source/Runtime/Foliage/Private/ProceduralFoliageTile.cpp b/Engine/Source/Runtime/Foliage/Private/ProceduralFoliageTile.cpp
--- a/Engine/Source/Runtime/Foliage/Private/ProceduralFoliageTile.cpp
+++ b/Engine/Source/Runtime/Foliage/Private/ProceduralFoliageTile.cpp
@@ -363,7 +363,9 @@ void UProceduralFoliageTile::CreateInstancesToSpawn(TArray<FProceduralFoliageIns
 			const UFoliageType_InstancedStaticMesh* Type = Instance.Type;
 			if (Hit.ImpactPoint.Z >= Type->HeightMin && Hit.ImpactPoint.Z <= Type->HeightMax)
 			{
-				if (FMath::Cos(FMath::DegreesToRadians(Type->GroundSlope)) <= Hit.ImpactNormal.Z)
+				const float MaxNormalAngle = FMath::Cos(FMath::DegreesToRadians(Type->GroundSlope));
+				const float MinNormalAngle = FMath::Cos(FMath::DegreesToRadians(Type->MinGroundSlope));
+				if (MaxNormalAngle <= Hit.ImpactNormal.Z && MinNormalAngle >= Hit.ImpactNormal.Z)	//keep in mind Hit.ImpactNormal.Z is (0,0,1) dot normal. However, ground slope is with relation to plane vector, not plane normal - so we swap comparisons
 				{
 					FProceduralFoliageInstance* NewInst = new(OutInstances)FProceduralFoliageInstance(Instance);
 					NewInst->Location = FVector(StartRay.X, StartRay.Y, Hit.ImpactPoint.Z);	//take the x,y of the instance, but use the z of the impact point. This is because we never want to move the instance along xy or we'll get overlaps 
diff --git a/Engine/Source/Runtime/Foliage/Public/FoliageType.h b/Engine/Source/Runtime/Foliage/Public/FoliageType.h
--- a/Engine/Source/Runtime/Foliage/Public/FoliageType.h
+++ b/Engine/Source/Runtime/Foliage/Public/FoliageType.h
@@ -67,6 +67,9 @@ class UFoliageType : public UObject
 	UPROPERTY(EditAnywhere, Category = General, meta = (Subcategory = "Placement"))
 	float GroundSlope;
 
+	UPROPERTY(EditAnywhere, Category = General, meta = (Subcategory = "Placement"))
+	float MinGroundSlope;
+
 	UPROPERTY(EditAnywhere, Category = General, meta = (Subcategory = "Placement"))
 	float HeightMin;
 
