Introduce per foliage type seed distribution. Allows for different foliage types to have close initial seeds

diff --git a/Engine/Source/Runtime/Foliage/Private/ProceduralFoliageTile.cpp b/Engine/Source/Runtime/Foliage/Private/ProceduralFoliageTile.cpp
--- a/Engine/Source/Runtime/Foliage/Private/ProceduralFoliageTile.cpp
+++ b/Engine/Source/Runtime/Foliage/Private/ProceduralFoliageTile.cpp
@@ -192,26 +192,106 @@ void UProceduralFoliageTile::AddRandomSeeds(TArray<FProceduralFoliageInstance*>&
 {
 	const float SizeTenM2 = (ProceduralFoliage->TileSize * ProceduralFoliage->TileSize) / (1000.f * 1000.f);
 
-	for ( const FProceduralFoliageTypeData& Data : ProceduralFoliage->GetTypes() )
+	const TArray<FProceduralFoliageTypeData>& Types = ProceduralFoliage->GetTypes();
+
+	TMap<int32,float> MaxShadeRadii;
+	TMap<int32, float> MaxCollisionRadii;
+	TMap<const UFoliageType*, int32> SeedsLeftMap;
+	TMap<const UFoliageType*, FRandomStream> RandomStreamPerType;
+
+	int32 TypesFinished = 0;
+
+	for (const FProceduralFoliageTypeData& Data : Types)
 	{
-		UFoliageType_InstancedStaticMesh* TypeInstance = Data.TypeInstance;
-		if( TypeInstance )
+		if (const UFoliageType_InstancedStaticMesh* Type = Data.TypeInstance)
 		{
-			const float NumSeeds = TypeInstance->GetSeedDensitySquared() * SizeTenM2;
-			for(int32 i = 0; i < NumSeeds; ++i)
-			{
-				const float X = RandomStream.FRandRange(0, ProceduralFoliage->TileSize);
-				const float Y = RandomStream.FRandRange(0, ProceduralFoliage->TileSize);
-				const float NewAge = TypeInstance->GetInitAge(RandomStream);
-				const float Scale = TypeInstance->GetScaleForAge(NewAge);
+			{	//compute the number of initial seeds
+				const int32 NumSeeds = FMath::RoundToInt(Type->GetSeedDensitySquared() * SizeTenM2);
+				SeedsLeftMap.Add(Type, NumSeeds);
+				if (NumSeeds == 0)
+				{
+					++TypesFinished;
+				}
+			}
 
-				if(FProceduralFoliageInstance* NewInst = NewSeed(FVector(X,Y,0.f), Scale, TypeInstance, NewAge))
+			{	//save the random stream per type
+				RandomStreamPerType.Add(Type, FRandomStream(Type->DistributionSeed));
+			}
+
+			{	//compute the needed offsets for initial seed variance
+				const int32 DistributionSeed = Type->DistributionSeed;
+				const float MaxScale = Type->GetScaleForAge(Type->MaxAge);
+				const float TypeMaxCollisionRadius = MaxScale * Type->CollisionRadius;
+				if (float* MaxRadius = MaxCollisionRadii.Find(DistributionSeed))
 				{
-					OutInstances.Add(NewInst);
+					*MaxRadius = FMath::Max(*MaxRadius, TypeMaxCollisionRadius);
+				}
+				else
+				{
+					MaxCollisionRadii.Add(DistributionSeed, TypeMaxCollisionRadius);
+				}
+
+				const float TypeMaxShadeRadius = MaxScale * Type->ShadeRadius;
+				if (float* MaxRadius = MaxShadeRadii.Find(DistributionSeed))
+				{
+					*MaxRadius = FMath::Max(*MaxRadius, TypeMaxShadeRadius);
+				}
+				else
+				{
+					MaxShadeRadii.Add(DistributionSeed, TypeMaxShadeRadius);
 				}
 			}
+			
 		}
 	}
+
+	int32 TypeIdx = -1;
+	const int32 NumTypes = Types.Num();
+	while (TypesFinished != NumTypes)
+	{
+		TypeIdx = (TypeIdx + 1) % NumTypes;	//keep cycling through the types that we spawn initial seeds for to make sure everyone gets fair chance
+
+		const FProceduralFoliageTypeData& Data = Types[TypeIdx];
+		if (const UFoliageType_InstancedStaticMesh* Type = Data.TypeInstance)
+		{
+			int32& SeedsLeft = SeedsLeftMap.FindChecked(Type);
+			if (SeedsLeft == 0)
+			{
+				continue;
+			}
+
+			FRandomStream& TypeRandomStream = RandomStreamPerType.FindChecked(Type);
+			const float NumSeeds = Type->GetSeedDensitySquared() * SizeTenM2;
+			const float InitX = TypeRandomStream.FRandRange(0, ProceduralFoliage->TileSize);
+			const float InitY = TypeRandomStream.FRandRange(0, ProceduralFoliage->TileSize);
+
+			const float Rad = RandomStream.FRandRange(0, PI*2.f);
+			const float NeededRadius = Type->bGrowsInShade ? MaxCollisionRadii.FindRef(Type->DistributionSeed) : MaxShadeRadii.FindRef(Type->DistributionSeed);
+			const FVector GlobalOffset = (RandomStream.FRandRange(0, Type->MaxInitialSeedOffset) + NeededRadius) * FVector(FMath::Cos(Rad), FMath::Sin(Rad), 0.f);
+
+			const float X = InitX + GlobalOffset.X;
+			const float Y = InitY + GlobalOffset.Y;
+
+			const float NewAge = Type->GetInitAge(RandomStream);
+			const float Scale = Type->GetScaleForAge(NewAge);
+
+			if (FProceduralFoliageInstance* NewInst = NewSeed(FVector(X, Y, 0.f), Scale, Type, NewAge))
+			{
+				OutInstances.Add(NewInst);
+			}
+
+			--SeedsLeft;
+			if (SeedsLeft == 0)
+			{
+				++TypesFinished;
+			}
+		}
+	}
+
+	for ( const FProceduralFoliageTypeData& Data : ProceduralFoliage->GetTypes() )
+	{
+		
+	}
 }
 
 void UProceduralFoliageTile::MarkPendingRemoval(FProceduralFoliageInstance* ToRemove)
diff --git a/Engine/Source/Runtime/Foliage/Public/FoliageType.h b/Engine/Source/Runtime/Foliage/Public/FoliageType.h
--- a/Engine/Source/Runtime/Foliage/Public/FoliageType.h
+++ b/Engine/Source/Runtime/Foliage/Public/FoliageType.h
@@ -277,6 +277,14 @@ class UFoliageType : public UObject
 	UPROPERTY(Category = Procedural, EditAnywhere, meta = (Subcategory = "Clustering", ClampMin = "0", UIMin = "0"))
 	int32 SeedsPerStep;
 
+	/** The seed that determines placement of initial seeds. */
+	UPROPERTY(Category = Procedural, EditAnywhere, meta = (Subcategory = "Clustering"))
+	int32 DistributionSeed;
+
+	/** The seed that determines placement of initial seeds. */
+	UPROPERTY(Category = Procedural, EditAnywhere, meta = (Subcategory = "Clustering"))
+	float MaxInitialSeedOffset;
+
 	/** Whether the species can grow in shade. If this is true shade radius is ignored during overlap tests*/
 	UPROPERTY(Category = Procedural, EditAnywhere, meta = (Subcategory = "Growth"))
 	bool bGrowsInShade;
