Procedural foliage that grows in shade will actually spawn around shade casting objects now

diff --git a/Engine/Source/Runtime/Foliage/Private/ProceduralFoliageTile.cpp b/Engine/Source/Runtime/Foliage/Private/ProceduralFoliageTile.cpp
--- a/Engine/Source/Runtime/Foliage/Private/ProceduralFoliageTile.cpp
+++ b/Engine/Source/Runtime/Foliage/Private/ProceduralFoliageTile.cpp
@@ -128,7 +128,7 @@ void UProceduralFoliageTile::AgeSeeds()
 		if (Instance->IsAlive())
 		{
 			const UFoliageType_InstancedStaticMesh* Type = Instance->Type;
-			if (SimulationStep <= Type->NumSteps)
+			if (SimulationStep <= Type->NumSteps && Type->bGrowsInShade == bSimulateShadeGrowth)
 			{
 				const float CurrentAge = Instance->Age;
 				const float NewAge = Type->GetNextAge(Instance->Age, 1);
@@ -165,7 +165,7 @@ void UProceduralFoliageTile::SpreadSeeds(TArray<FProceduralFoliageInstance*>& Ne
 
 		const UFoliageType_InstancedStaticMesh* Type = Inst->Type;
 
-		if (SimulationStep <= Type->NumSteps)
+		if (SimulationStep <= Type->NumSteps  && Type->bGrowsInShade == bSimulateShadeGrowth)
 		{
 			for (int32 i = 0; i < Type->SeedsPerStep; ++i)
 			{
@@ -192,25 +192,24 @@ void UProceduralFoliageTile::AddRandomSeeds(TArray<FProceduralFoliageInstance*>&
 {
 	const float SizeTenM2 = (ProceduralFoliage->TileSize * ProceduralFoliage->TileSize) / (1000.f * 1000.f);
 
-	const TArray<FProceduralFoliageTypeData>& Types = ProceduralFoliage->GetTypes();
-
 	TMap<int32,float> MaxShadeRadii;
 	TMap<int32, float> MaxCollisionRadii;
 	TMap<const UFoliageType*, int32> SeedsLeftMap;
 	TMap<const UFoliageType*, FRandomStream> RandomStreamPerType;
 
-	int32 TypesFinished = 0;
+	TArray<const UFoliageType_InstancedStaticMesh*> TypesToSeed;
 
-	for (const FProceduralFoliageTypeData& Data : Types)
+	for (const FProceduralFoliageTypeData& Data : ProceduralFoliage->GetTypes())
 	{
-		if (const UFoliageType_InstancedStaticMesh* Type = Data.TypeInstance)
+		const UFoliageType_InstancedStaticMesh* Type = Data.TypeInstance;
+		if (Type && Type->bGrowsInShade == bSimulateShadeGrowth)
 		{
 			{	//compute the number of initial seeds
 				const int32 NumSeeds = FMath::RoundToInt(Type->GetSeedDensitySquared() * SizeTenM2);
 				SeedsLeftMap.Add(Type, NumSeeds);
-				if (NumSeeds == 0)
+				if (NumSeeds > 0)
 				{
-					++TypesFinished;
+					TypesToSeed.Add(Type);
 				}
 			}
 
@@ -246,13 +245,14 @@ void UProceduralFoliageTile::AddRandomSeeds(TArray<FProceduralFoliageInstance*>&
 	}
 
 	int32 TypeIdx = -1;
-	const int32 NumTypes = Types.Num();
-	while (TypesFinished != NumTypes)
+	const int32 NumTypes = TypesToSeed.Num();
+	int32 TypesLeftToSeed = NumTypes;
+	const int32 LastShadeCastingIndex = InstancesArray.Num() - 1; //when placing shade growth types we want to spawn in shade if possible
+	while (TypesLeftToSeed > 0)
 	{
 		TypeIdx = (TypeIdx + 1) % NumTypes;	//keep cycling through the types that we spawn initial seeds for to make sure everyone gets fair chance
 
-		const FProceduralFoliageTypeData& Data = Types[TypeIdx];
-		if (const UFoliageType_InstancedStaticMesh* Type = Data.TypeInstance)
+		if (const UFoliageType_InstancedStaticMesh* Type = TypesToSeed[TypeIdx])
 		{
 			int32& SeedsLeft = SeedsLeftMap.FindChecked(Type);
 			if (SeedsLeft == 0)
@@ -261,9 +261,21 @@ void UProceduralFoliageTile::AddRandomSeeds(TArray<FProceduralFoliageInstance*>&
 			}
 
 			FRandomStream& TypeRandomStream = RandomStreamPerType.FindChecked(Type);
-			const float NumSeeds = Type->GetSeedDensitySquared() * SizeTenM2;
-			const float InitX = TypeRandomStream.FRandRange(0, ProceduralFoliage->TileSize);
-			const float InitY = TypeRandomStream.FRandRange(0, ProceduralFoliage->TileSize);
+			float InitX = 0.f;
+			float InitY = 0.f;
+
+			if (bSimulateShadeGrowth && LastShadeCastingIndex >= 0)
+			{
+				const int32 InstanceSpawnerIdx = TypeRandomStream.FRandRange(0, LastShadeCastingIndex);
+				const FProceduralFoliageInstance& Spawner = InstancesArray[InstanceSpawnerIdx];
+				InitX = Spawner.Location.X;
+				InitY = Spawner.Location.Y;
+			}
+			else
+			{
+				InitX = TypeRandomStream.FRandRange(0, ProceduralFoliage->TileSize);
+				InitY = TypeRandomStream.FRandRange(0, ProceduralFoliage->TileSize);
+			}
 
 			const float Rad = RandomStream.FRandRange(0, PI*2.f);
 			const float NeededRadius = Type->bGrowsInShade ? MaxCollisionRadii.FindRef(Type->DistributionSeed) : MaxShadeRadii.FindRef(Type->DistributionSeed);
@@ -283,15 +295,10 @@ void UProceduralFoliageTile::AddRandomSeeds(TArray<FProceduralFoliageInstance*>&
 			--SeedsLeft;
 			if (SeedsLeft == 0)
 			{
-				++TypesFinished;
+				--TypesLeftToSeed;
 			}
 		}
 	}
-
-	for ( const FProceduralFoliageTypeData& Data : ProceduralFoliage->GetTypes() )
-	{
-		
-	}
 }
 
 void UProceduralFoliageTile::MarkPendingRemoval(FProceduralFoliageInstance* ToRemove)
@@ -378,17 +385,16 @@ void UProceduralFoliageTile::StepSimulation()
 	FlushPendingRemovals();
 }
 
-void UProceduralFoliageTile::Simulate(const UProceduralFoliage* InProceduralFoliage, const int32 RandomSeed, const int32 MaxNumSteps)
+void UProceduralFoliageTile::StartSimulation(const int32 MaxNumSteps, bool bShadeGrowth)
 {
-	InitSimulation(InProceduralFoliage, RandomSeed);
-
 	int32 MaxSteps = 0;
-	for(const FProceduralFoliageTypeData& Data : ProceduralFoliage->GetTypes())
+
+	for (const FProceduralFoliageTypeData& Data : ProceduralFoliage->GetTypes())
 	{
 		const UFoliageType_InstancedStaticMesh* TypeInstance = Data.TypeInstance;
-		if(TypeInstance)
+		if (TypeInstance && TypeInstance->bGrowsInShade == bShadeGrowth)
 		{
-			MaxSteps = FMath::Max(MaxSteps, TypeInstance->NumSteps+1);
+			MaxSteps = FMath::Max(MaxSteps, TypeInstance->NumSteps + 1);
 		}
 	}
 
@@ -397,6 +403,8 @@ void UProceduralFoliageTile::Simulate(const UProceduralFoliage* InProceduralFoli
 		MaxSteps = FMath::Min(MaxSteps, MaxNumSteps);	//only take as many steps as given
 	}
 
+	SimulationStep = 0;
+	bSimulateShadeGrowth = bShadeGrowth;
 	for (int32 Step = 0; Step < MaxSteps; ++Step)
 	{
 		StepSimulation();
@@ -406,6 +414,14 @@ void UProceduralFoliageTile::Simulate(const UProceduralFoliage* InProceduralFoli
 	InstancesToArray();
 }
 
+void UProceduralFoliageTile::Simulate(const UProceduralFoliage* InProceduralFoliage, const int32 RandomSeed, const int32 MaxNumSteps)
+{
+	InitSimulation(InProceduralFoliage, RandomSeed);
+
+	StartSimulation(MaxNumSteps, false);
+	StartSimulation(MaxNumSteps, true);
+}
+
 
 void UProceduralFoliageTile::BeginDestroy()
 {
diff --git a/Engine/Source/Runtime/Foliage/Public/ProceduralFoliageTile.h b/Engine/Source/Runtime/Foliage/Public/ProceduralFoliageTile.h
--- a/Engine/Source/Runtime/Foliage/Public/ProceduralFoliageTile.h
+++ b/Engine/Source/Runtime/Foliage/Public/ProceduralFoliageTile.h
@@ -42,6 +42,7 @@ class FOLIAGE_API UProceduralFoliageTile : public UObject
 	void Empty();
 private:
 	
+	void StartSimulation(const int32 MaxNumSteps, bool bShadeGrowth);
 	void StepSimulation();
 	void AddRandomSeeds(TArray<FProceduralFoliageInstance*>& OutInstances);
 
@@ -69,6 +70,7 @@ private:
 	FProceduralFoliageBroadphase Broadphase;	
 
 	FRandomStream RandomStream;
+	bool bSimulateShadeGrowth;
 
 private:
 	float GetRandomGaussian();
