Grows in shade procedural foliage is spawned right next to a tree

diff --git a/Engine/Source/Runtime/Foliage/Private/ProceduralFoliageTile.cpp b/Engine/Source/Runtime/Foliage/Private/ProceduralFoliageTile.cpp
--- a/Engine/Source/Runtime/Foliage/Private/ProceduralFoliageTile.cpp
+++ b/Engine/Source/Runtime/Foliage/Private/ProceduralFoliageTile.cpp
@@ -260,9 +260,13 @@ void UProceduralFoliageTile::AddRandomSeeds(TArray<FProceduralFoliageInstance*>&
 				continue;
 			}
 
+			const float NewAge = Type->GetInitAge(RandomStream);
+			const float Scale = Type->GetScaleForAge(NewAge);
+
 			FRandomStream& TypeRandomStream = RandomStreamPerType.FindChecked(Type);
 			float InitX = 0.f;
 			float InitY = 0.f;
+			float NeededRadius = 0.f;
 
 			if (bSimulateShadeGrowth && LastShadeCastingIndex >= 0)
 			{
@@ -270,23 +274,23 @@ void UProceduralFoliageTile::AddRandomSeeds(TArray<FProceduralFoliageInstance*>&
 				const FProceduralFoliageInstance& Spawner = InstancesArray[InstanceSpawnerIdx];
 				InitX = Spawner.Location.X;
 				InitY = Spawner.Location.Y;
+				NeededRadius = Spawner.GetCollisionRadius() * (Scale + Spawner.Type->GetScaleForAge(Spawner.Age));
 			}
 			else
 			{
 				InitX = TypeRandomStream.FRandRange(0, ProceduralFoliage->TileSize);
 				InitY = TypeRandomStream.FRandRange(0, ProceduralFoliage->TileSize);
+				NeededRadius = MaxShadeRadii.FindRef(Type->DistributionSeed);
 			}
 
 			const float Rad = RandomStream.FRandRange(0, PI*2.f);
-			const float NeededRadius = Type->bGrowsInShade ? MaxCollisionRadii.FindRef(Type->DistributionSeed) : MaxShadeRadii.FindRef(Type->DistributionSeed);
+			
+			
 			const FVector GlobalOffset = (RandomStream.FRandRange(0, Type->MaxInitialSeedOffset) + NeededRadius) * FVector(FMath::Cos(Rad), FMath::Sin(Rad), 0.f);
 
 			const float X = InitX + GlobalOffset.X;
 			const float Y = InitY + GlobalOffset.Y;
 
-			const float NewAge = Type->GetInitAge(RandomStream);
-			const float Scale = Type->GetScaleForAge(NewAge);
-
 			if (FProceduralFoliageInstance* NewInst = NewSeed(FVector(X, Y, 0.f), Scale, Type, NewAge))
 			{
 				OutInstances.Add(NewInst);
