Allow blocking volumes to block specific foliage actors

diff --git a/Engine/Source/Runtime/Foliage/Private/InstancedFoliage.cpp b/Engine/Source/Runtime/Foliage/Private/InstancedFoliage.cpp
--- a/Engine/Source/Runtime/Foliage/Private/InstancedFoliage.cpp
+++ b/Engine/Source/Runtime/Foliage/Private/InstancedFoliage.cpp
@@ -14,6 +14,7 @@
 #include "Serialization/CustomVersion.h"
 #include "ProceduralFoliageComponent.h"
 #include "ProceduralFoliageBlockingVolume.h"
+#include "ProceduralFoliageActor.h"
 
 #define LOCTEXT_NAMESPACE "InstancedFoliage"
 
@@ -1848,7 +1849,15 @@ AInstancedFoliageActor* AInstancedFoliageActor::FoliageTrace(UWorld* InWorld, FH
 		{
 			if (Hit.Actor.IsValid())	//if we hit the ProceduralFoliage blocking volume don't spawn instance
 			{
-				if (Cast<AProceduralFoliageBlockingVolume>(Hit.Actor.Get()) || Cast<AInstancedFoliageActor>(Hit.Actor.Get()))
+				if (AProceduralFoliageBlockingVolume* ProceduralFoliageBlockingVolume = Cast<AProceduralFoliageBlockingVolume>(Hit.Actor.Get()))
+				{
+					const AProceduralFoliageActor* ProceduralFoliageActor = ProceduralFoliageBlockingVolume->ProceduralFoliageActor;
+					if (ProceduralFoliageActor == nullptr || ProceduralFoliageActor->ProceduralComponent == nullptr || ProceduralFoliageActor->ProceduralComponent->GetProceduralGuid() == DesiredInstance.ProceduralGuid)
+					{
+						return nullptr;
+					}
+
+				}else if(Cast<AInstancedFoliageActor>(Hit.Actor.Get()))
 				{
 					return nullptr;
 				}
diff --git a/Engine/Source/Runtime/Foliage/Private/ProceduralFoliageBlockingVolume.h b/Engine/Source/Runtime/Foliage/Private/ProceduralFoliageBlockingVolume.h
--- a/Engine/Source/Runtime/Foliage/Private/ProceduralFoliageBlockingVolume.h
+++ b/Engine/Source/Runtime/Foliage/Private/ProceduralFoliageBlockingVolume.h
@@ -5,13 +5,16 @@
 #include "GameFramework/Volume.h"
 #include "ProceduralFoliageBlockingVolume.generated.h"
 
+class AProceduralFoliageActor;
+
 /** An invisible volume used to block ProceduralFoliage instances from being spawned. */
 UCLASS(MinimalAPI)
 class AProceduralFoliageBlockingVolume : public AVolume
 {
 	GENERATED_UCLASS_BODY()
 
-	// End UObject interface.
+	UPROPERTY(Category = ProceduralFoliage, EditAnywhere)
+	AProceduralFoliageActor* ProceduralFoliageActor;
 };
 
 
