Foliage supports cross-level base references
By default paint tool adds instances into base compoent level instead of current
UE-6556

diff --git a/Engine/Source/Editor/FoliageEdit/Private/FoliageEdMode.cpp b/Engine/Source/Editor/FoliageEdit/Private/FoliageEdMode.cpp
--- a/Engine/Source/Editor/FoliageEdit/Private/FoliageEdMode.cpp
+++ b/Engine/Source/Editor/FoliageEdit/Private/FoliageEdMode.cpp
@@ -271,16 +271,14 @@ void FEdModeFoliage::FoliageBrushTrace(FEditorViewportClient* ViewportClient, in
 
 			FHitResult Hit;
 			UWorld* World = ViewportClient->GetWorld();
-			AInstancedFoliageActor* IFA = AInstancedFoliageActor::GetInstancedFoliageActorForCurrentLevel(World, false);
 			static FName NAME_FoliageBrush = FName(TEXT("FoliageBrush"));
-			if (AInstancedFoliageActor::FoliageTrace(World, Hit, FDesiredFoliageInstance(Start, End), IFA, NAME_FoliageBrush))
+			if (AInstancedFoliageActor::FoliageTrace(World, Hit, FDesiredFoliageInstance(Start, End), nullptr, NAME_FoliageBrush))
 			{
 				// Check filters
 				UPrimitiveComponent* PrimComp = Hit.Component.Get();
 				UMaterialInterface* Material = PrimComp ? PrimComp->GetMaterial(0) : nullptr;
 
 				if (PrimComp &&
-					PrimComp->GetOutermost() == World->GetCurrentLevel()->GetOutermost() &&
 					(UISettings.bFilterLandscape || !PrimComp->IsA(ULandscapeHeightfieldCollisionComponent::StaticClass())) &&
 					(UISettings.bFilterStaticMesh || !PrimComp->IsA(UStaticMeshComponent::StaticClass())) &&
 					(UISettings.bFilterBSP || !PrimComp->IsA(UModelComponent::StaticClass())) &&
@@ -502,54 +500,49 @@ void FEdModeFoliage::CalculatePotentialInstances(UWorld* InWorld, const AInstanc
 	{
 		FHitResult Hit;
 		static FName NAME_AddFoliageInstances = FName(TEXT("AddFoliageInstances"));
-		if (AInstancedFoliageActor* IFA = AInstancedFoliageActor::FoliageTrace(InWorld, Hit, DesiredInst, IgnoreIFA, NAME_AddFoliageInstances, true))
+		if (AInstancedFoliageActor::FoliageTrace(InWorld, Hit, DesiredInst, IgnoreIFA, NAME_AddFoliageInstances, true))
 		{
 			// Check filters
+			UPrimitiveComponent* InstanceBase = Hit.Component.Get();
+			check(InstanceBase);
 
-			if (UPrimitiveComponent* PrimComp = Hit.Component.Get())
+			if (DesiredInst.PlacementMode == EFoliagePlacementMode::Manual)
 			{
-				if (DesiredInst.PlacementMode == EFoliagePlacementMode::Manual)
+				UMaterialInterface* Material = InstanceBase->GetMaterial(0);
+				if ((!UISettings->bFilterLandscape && InstanceBase->IsA(ULandscapeHeightfieldCollisionComponent::StaticClass())) ||
+					(!UISettings->bFilterStaticMesh && InstanceBase->IsA(UStaticMeshComponent::StaticClass())) ||
+					(!UISettings->bFilterBSP && InstanceBase->IsA(UModelComponent::StaticClass())) ||
+					(!UISettings->bFilterTranslucent && Material && IsTranslucentBlendMode(Material->GetBlendMode()))
+					)
 				{
-					UMaterialInterface* Material = PrimComp ? PrimComp->GetMaterial(0) : nullptr;
-					if (PrimComp->GetOutermost() != InWorld->GetCurrentLevel()->GetOutermost() ||
-						(!UISettings->bFilterLandscape && PrimComp->IsA(ULandscapeHeightfieldCollisionComponent::StaticClass())) ||
-						(!UISettings->bFilterStaticMesh && PrimComp->IsA(UStaticMeshComponent::StaticClass())) ||
-						(!UISettings->bFilterBSP && PrimComp->IsA(UModelComponent::StaticClass())) ||
-						(!UISettings->bFilterTranslucent && Material && IsTranslucentBlendMode(Material->GetBlendMode()))
-						)
-					{
-						continue;
-					}
+					continue;
 				}
 			}
-			else
+
+			// New foliage instances always go into base component level
+			AInstancedFoliageActor* TargetIFA = AInstancedFoliageActor::GetInstancedFoliageActorForLevel(InstanceBase->GetComponentLevel(), true);
+			if (TargetIFA == nullptr)
 			{
 				continue;
 			}
-
-
+				
 			FFoliageMeshInfo* MeshInfo = DesiredInst.MeshInfo;
-			if (DesiredInst.PlacementMode == EFoliagePlacementMode::Procedural)
+			//We have to add the mesh into the foliage actor. This assumes the foliage type is just a concrete InstancedStaticMesh.
+			//With the new tree system coming this will not be true and you'll need to traverse the tree and insert the needed meshes
+			if (UStaticMesh* StaticMesh = Settings->GetStaticMesh())
 			{
-				//We have to add the mesh into the foliage actor. This assumes the foliage type is just a concrete InstancedStaticMesh.
-				//With the new tree system coming this will not be true and you'll need to traverse the tree and insert the needed meshes
-				if (UStaticMesh* StaticMesh = Settings->GetStaticMesh())
+				UFoliageType* ExistingSettings = TargetIFA->GetSettingsForMesh(StaticMesh);
+				if (ExistingSettings == nullptr)
 				{
-					UFoliageType* ExistingSettings = IFA->GetSettingsForMesh(StaticMesh);
-					if (ExistingSettings == nullptr)
-					{
-						if (const UFoliageType_InstancedStaticMesh* InstancedStaticMeshFoliageType = Cast<UFoliageType_InstancedStaticMesh>(Settings))
-						{
-							MeshInfo = IFA->AddMesh(Settings->GetStaticMesh(), nullptr, InstancedStaticMeshFoliageType);
-						}
-					}
-					else
+					if (const UFoliageType_InstancedStaticMesh* InstancedStaticMeshFoliageType = Cast<UFoliageType_InstancedStaticMesh>(Settings))
 					{
-						MeshInfo = IFA->FindMesh(ExistingSettings);
+						MeshInfo = TargetIFA->AddMesh(Settings->GetStaticMesh(), nullptr, InstancedStaticMeshFoliageType);
 					}
 				}
-
-				
+				else
+				{
+					MeshInfo = TargetIFA->FindMesh(ExistingSettings);
+				}
 			}
 
 			if (MeshInfo == nullptr)
@@ -565,7 +558,7 @@ void FEdModeFoliage::CalculatePotentialInstances(UWorld* InWorld, const AInstanc
 			// Check vertex color mask
 			if (Settings->VertexColorMask != FOLIAGEVERTEXCOLORMASK_Disabled && Hit.FaceIndex != INDEX_NONE)
 			{
-				if (UStaticMeshComponent* HitStaticMeshComponent = Cast<UStaticMeshComponent>(Hit.Component.Get()))
+				if (UStaticMeshComponent* HitStaticMeshComponent = Cast<UStaticMeshComponent>(InstanceBase))
 				{
 					FColor VertexColor;
 					if (GetStaticMeshVertexColorForHit(HitStaticMeshComponent, Hit.FaceIndex, Hit.ImpactPoint, VertexColor))
@@ -580,7 +573,7 @@ void FEdModeFoliage::CalculatePotentialInstances(UWorld* InWorld, const AInstanc
 
 			// Check landscape layer
 			float HitWeight = 1.f;
-			if (LandscapeLayersValid(Settings) && GetMaxHitWeight(Hit.ImpactPoint, Hit.Component.Get(), Settings, LandscapeLayerCachesPtr, HitWeight))
+			if (LandscapeLayersValid(Settings) && GetMaxHitWeight(Hit.ImpactPoint, InstanceBase, Settings, LandscapeLayerCachesPtr, HitWeight))
 			{
 				// Reject instance randomly in proportion to weight
 				if (FilterByWeight(HitWeight, Settings))
@@ -590,7 +583,7 @@ void FEdModeFoliage::CalculatePotentialInstances(UWorld* InWorld, const AInstanc
 			}
 
 			const int32 BucketIndex = FMath::RoundToInt(HitWeight * (float)(NUM_INSTANCE_BUCKETS - 1));
-			OutPotentialInstances[BucketIndex].Add(FPotentialInstance(Hit.ImpactPoint, Hit.ImpactNormal, Hit.Component.Get(), HitWeight, MeshInfo, IFA, DesiredInst));
+			OutPotentialInstances[BucketIndex].Add(FPotentialInstance(Hit.ImpactPoint, Hit.ImpactNormal, InstanceBase, HitWeight, MeshInfo, TargetIFA, DesiredInst));
 		}
 	}
 }
@@ -656,8 +649,9 @@ void FEdModeFoliage::AddInstancesForBrush(UWorld* InWorld, AInstancedFoliageActo
 			for (int32 Idx : ExistingInstances)
 			{
 				FFoliageInstance& Instance = MeshInfo.Instances[Idx];
+				auto InstanceBasePtr = IFA->InstanceBaseCache.GetInstanceBasePtr(Instance.BaseId);
 				float HitWeight;
-				if (GetMaxHitWeight(Instance.Location, Instance.Base, Settings, &LandscapeLayerCaches, HitWeight))
+				if (GetMaxHitWeight(Instance.Location, InstanceBasePtr.Get(), Settings, &LandscapeLayerCaches, HitWeight))
 				{
 					// Add count to bucket.
 					ExistingInstanceBuckets[FMath::RoundToInt(HitWeight * (float)(NUM_INSTANCE_BUCKETS - 1))]++;
@@ -705,7 +699,8 @@ void FEdModeFoliage::RemoveInstancesForBrush(AInstancedFoliageActor* IFA, FFolia
 		// Filter PotentialInstancesToRemove
 		for (int32 Idx = 0; Idx < PotentialInstancesToRemove.Num(); Idx++)
 		{
-			UPrimitiveComponent* Base = Cast<UPrimitiveComponent> (MeshInfo.Instances[PotentialInstancesToRemove[Idx]].Base);
+			auto BasePtr = IFA->InstanceBaseCache.GetInstanceBasePtr(MeshInfo.Instances[PotentialInstancesToRemove[Idx]].BaseId);
+			UPrimitiveComponent* Base = Cast<UPrimitiveComponent>(BasePtr.Get());
 			UMaterialInterface* Material = Base ? Base->GetMaterial(0) : nullptr;
 
 			// Check if instance is candidate for removal based on filter settings
@@ -1213,8 +1208,10 @@ void FEdModeFoliage::ApplyPaintBucket(AActor* Actor, bool bRemove)
 		TInlineComponentArray<UActorComponent*> Components;
 		Actor->GetComponents(Components);
 
-		for (int32 ComponentIdx = 0; ComponentIdx < Components.Num(); ComponentIdx++)
+		for (auto Component : Components)
 		{
+			auto BaseId = IFA->InstanceBaseCache.GetInstanceBaseId(Component);
+			
 			for (auto& MeshPair : IFA->FoliageMeshes)
 			{
 				FFoliageMeshInfo& MeshInfo = *MeshPair.Value;
@@ -1222,10 +1219,10 @@ void FEdModeFoliage::ApplyPaintBucket(AActor* Actor, bool bRemove)
 
 				if (Settings->IsSelected)
 				{
-					FFoliageComponentHashInfo* ComponentHashInfo = MeshInfo.ComponentHash.Find(Components[ComponentIdx]);
-					if (ComponentHashInfo)
+					auto* InstanceSet = MeshInfo.ComponentHash.Find(BaseId);
+					if (InstanceSet)
 					{
-						TArray<int32> InstancesToRemove = ComponentHashInfo->Instances.Array();
+						TArray<int32> InstancesToRemove = InstanceSet->Array();
 						MeshInfo.RemoveInstances(IFA, InstancesToRemove);
 					}
 				}
@@ -1479,6 +1476,53 @@ bool FEdModeFoliage::GetStaticMeshVertexColorForHit(UStaticMeshComponent* InStat
 	return false;
 }
 
+bool FEdModeFoliage::SnapInstanceToGround(AInstancedFoliageActor* InIFA, float AlignMaxAngle, FFoliageMeshInfo& Mesh, int32 InstanceIdx)
+{
+	FFoliageInstance& Instance = Mesh.Instances[InstanceIdx];
+	FVector Start = Instance.Location;
+	FVector End = Instance.Location - FVector(0.f, 0.f, FOLIAGE_SNAP_TRACE);
+
+	FHitResult Hit;
+	static FName NAME_FoliageSnap = FName("FoliageSnap");
+	if (AInstancedFoliageActor::FoliageTrace(InIFA->GetWorld(), Hit, FDesiredFoliageInstance(Start, End), InIFA, NAME_FoliageSnap))
+	{
+		UPrimitiveComponent* HitComponent = Hit.Component.Get();
+
+		// We cannot be based on an a blueprint component as these will disappear when the construction script is re-run
+		if (HitComponent->CreationMethod == EComponentCreationMethod::ConstructionScript)
+		{
+			HitComponent = nullptr;
+		}
+																
+		// Find BSP brush 
+		UModelComponent* ModelComponent = Cast<UModelComponent>(HitComponent);
+		if (ModelComponent)
+		{
+			ABrush* BrushActor = ModelComponent->GetModel()->FindBrush(Hit.Location);
+			if (BrushActor)
+			{
+				HitComponent = BrushActor->GetBrushComponent();
+			}
+		}
+
+		auto NewBaseId = InIFA->InstanceBaseCache.AddInstanceBaseId(HitComponent);
+		Mesh.SetInstanceBaseId(NewBaseId, InstanceIdx);
+	
+		Instance.Location = Hit.Location;
+		Instance.ZOffset = 0.f;
+											
+		if (Instance.Flags & FOLIAGE_AlignToNormal)
+		{
+			// Remove previous alignment and align to new normal.
+			Instance.Rotation = Instance.PreAlignRotation;
+			Instance.AlignToNormal(Hit.Normal, AlignMaxAngle);
+		}
+
+		return true;
+	}
+
+	return false;
+}
 
 
 /** Get list of meshs for current level for UI */
@@ -1853,42 +1897,10 @@ bool FEdModeFoliage::InputKey(FEditorViewportClient* ViewportClient, FViewport*
 					TArray<int32> SelectedIndices = Mesh.SelectedIndices.Array();
 
 					Mesh.PreMoveInstances(IFA, SelectedIndices);
-
-					UWorld* World = ViewportClient->GetWorld();
-					for (int32 SelectedInstanceIdx : SelectedIndices)
+					for (int32 InstanceIndex : SelectedIndices)
 					{
-						FFoliageInstance& Instance = Mesh.Instances[SelectedInstanceIdx];
-
-						FVector Start = Instance.Location;
-						FVector End = Instance.Location - FVector(0.f, 0.f, FOLIAGE_SNAP_TRACE);
-
-						FHitResult Hit;
-						if (AInstancedFoliageActor::FoliageTrace(GetWorld(), Hit, FDesiredFoliageInstance(Start, End), IFA, FName("FoliageSnap")))
-						{
-							// Check current level
-							if ((Hit.Component.IsValid() && Hit.Component.Get()->GetOutermost() == World->GetCurrentLevel()->GetOutermost()) ||
-								(Hit.GetActor() && Hit.GetActor()->IsA(AWorldSettings::StaticClass())))
-							{
-								Instance.Location = Hit.Location;
-								Instance.ZOffset = 0.f;
-								Instance.Base = Hit.Component.Get();
-								// We cannot be based on an a blueprint component as these will disappear when the construction script is re-run
-								if (Instance.Base && Instance.Base->CreationMethod == EComponentCreationMethod::ConstructionScript)
-								{
-									Instance.Base = nullptr;
-								}
-
-								if (Instance.Flags & FOLIAGE_AlignToNormal)
-								{
-									// Remove previous alignment and align to new normal.
-									Instance.Rotation = Instance.PreAlignRotation;
-									Instance.AlignToNormal(Hit.Normal, MeshPair.Key->AlignMaxAngle);
-								}
-							}
-						}
-						bMovedInstance = true;
+						bMovedInstance|= SnapInstanceToGround(IFA, MeshPair.Key->AlignMaxAngle, Mesh, InstanceIndex);
 					}
-
 					Mesh.PostMoveInstances(IFA, SelectedIndices);
 				}
 
diff --git a/Engine/Source/Editor/FoliageEdit/Private/FoliageEdMode.h b/Engine/Source/Editor/FoliageEdit/Private/FoliageEdMode.h
--- a/Engine/Source/Editor/FoliageEdit/Private/FoliageEdMode.h
+++ b/Engine/Source/Editor/FoliageEdit/Private/FoliageEdMode.h
@@ -329,6 +329,12 @@ private:
 	/** Lookup the vertex color corresponding to a location traced on a static mesh */
 	static bool GetStaticMeshVertexColorForHit(UStaticMeshComponent* InStaticMeshComponent, int32 InTriangleIndex, const FVector& InHitLocation, FColor& OutVertexColor);
 
+	/** 
+	 * Snaps given instance to the ground  
+	 * @return Whether instance was sucessfully snapped
+	 */
+	bool SnapInstanceToGround(AInstancedFoliageActor* InIFA, float AlignMaxAngle, FFoliageMeshInfo& Mesh, int32 InstanceIdx);
+
 	bool bBrushTraceValid;
 	FVector BrushLocation;
 	FVector BrushTraceDirection;
diff --git a/Engine/Source/Editor/UnrealEd/Private/EditorActor.cpp b/Engine/Source/Editor/UnrealEd/Private/EditorActor.cpp
--- a/Engine/Source/Editor/UnrealEd/Private/EditorActor.cpp
+++ b/Engine/Source/Editor/UnrealEd/Private/EditorActor.cpp
@@ -680,19 +680,6 @@ bool UUnrealEdEngine::edactDeleteSelected( UWorld* InWorld, bool bVerifyDeletion
 			Level->Modify();
 		}
 
-		// See if there is any foliage that also needs to be removed
-		AInstancedFoliageActor* IFA = AInstancedFoliageActor::GetInstancedFoliageActorForLevel(Level);
-		if( IFA )
-		{
-			TInlineComponentArray<UActorComponent*> Components;
-			Actor->GetComponents(Components);
-
-			for(int32 ComponentIndex = 0;ComponentIndex < Components.Num();ComponentIndex++)
-			{
-				IFA->DeleteInstancesForComponent( Components[ComponentIndex] );
-			}
-		}
-
 		UE_LOG(LogEditorActor, Log,  TEXT("Deleted Actor: %s"), *Actor->GetClass()->GetName() );
 
 		// Destroy actor and clear references.
diff --git a/Engine/Source/Editor/UnrealEd/Private/EditorObject.cpp b/Engine/Source/Editor/UnrealEd/Private/EditorObject.cpp
--- a/Engine/Source/Editor/UnrealEd/Private/EditorObject.cpp
+++ b/Engine/Source/Editor/UnrealEd/Private/EditorObject.cpp
@@ -284,8 +284,6 @@ static const TCHAR* ImportProperties(
 							}
 							FParse::Value(StrPtr, TEXT("Flags="), Instance.Flags);
 
-							Instance.Base = ActorComponent;
-
 							// Add the instance
 							FFoliageMeshInfo* MeshInfo;
 							UFoliageType* Type = IFA->GetSettingsForMesh(StaticMesh, &MeshInfo);
@@ -293,7 +291,7 @@ static const TCHAR* ImportProperties(
 							{
 								MeshInfo = IFA->AddMesh(StaticMesh, &Type);
 							}
-							MeshInfo->AddInstance(IFA, Type, Instance);
+							MeshInfo->AddInstance(IFA, Type, Instance, ActorComponent);
 						}
 					}
 				}
diff --git a/Engine/Source/Editor/WorldBrowser/Private/Tiles/WorldTileModel.cpp b/Engine/Source/Editor/WorldBrowser/Private/Tiles/WorldTileModel.cpp
--- a/Engine/Source/Editor/WorldBrowser/Private/Tiles/WorldTileModel.cpp
+++ b/Engine/Source/Editor/WorldBrowser/Private/Tiles/WorldTileModel.cpp
@@ -466,6 +466,11 @@ void FWorldTileModel::SetLevelPosition(const FIntPoint& InPosition)
 		if (Offset != FIntPoint::ZeroValue)
 		{
 			Level->ApplyWorldOffset(FVector(Offset), false);
+
+			for (auto Actor : Level->Actors)
+			{
+				Editor->BroadcastOnActorMoved(Actor);
+			}
 		}
 	}
 
diff --git a/Engine/Source/Runtime/Engine/Classes/Engine/WorldComposition.h b/Engine/Source/Runtime/Engine/Classes/Engine/WorldComposition.h
--- a/Engine/Source/Runtime/Engine/Classes/Engine/WorldComposition.h
+++ b/Engine/Source/Runtime/Engine/Classes/Engine/WorldComposition.h
@@ -126,6 +126,9 @@ class ENGINE_API UWorldComposition : public UObject
 	/**  */
 	void ReinitializeForPIE();
 
+	/** @returns Whether specified tile package name is managed by world composition */
+	bool DoesTileExists(const FName& TilePackageName) const;
+
 #if WITH_EDITOR
 	/** @returns FWorldTileInfo associated with specified package */
 	FWorldTileInfo GetTileInfo(const FName& InPackageName) const;
diff --git a/Engine/Source/Runtime/Engine/Private/WorldComposition.cpp b/Engine/Source/Runtime/Engine/Private/WorldComposition.cpp
--- a/Engine/Source/Runtime/Engine/Private/WorldComposition.cpp
+++ b/Engine/Source/Runtime/Engine/Private/WorldComposition.cpp
@@ -275,6 +275,18 @@ void UWorldComposition::ReinitializeForPIE()
 	GetWorld()->StreamingLevels.Append(TilesStreaming);
 }
 
+bool UWorldComposition::DoesTileExists(const FName& InTilePackageName) const
+{
+	for (const auto& Tile : Tiles)
+	{
+		if (Tile.PackageName == InTilePackageName)
+		{
+			return true;
+		}
+	}
+	return false;
+}
+
 ULevelStreaming* UWorldComposition::CreateStreamingLevel(const FWorldCompositionTile& InTile) const
 {
 	UWorld* OwningWorld = GetWorld();
diff --git a/Engine/Source/Runtime/Foliage/Private/FoliageInstanceBase.cpp b/Engine/Source/Runtime/Foliage/Private/FoliageInstanceBase.cpp
new file mode 100644
--- /dev/null
+++ b/Engine/Source/Runtime/Foliage/Private/FoliageInstanceBase.cpp
@@ -0,0 +1,280 @@
+// Copyright 1998-2015 Epic Games, Inc. All Rights Reserved.
+#include "FoliagePrivate.h"
+#include "FoliageInstanceBase.h"
+#include "InstancedFoliage.h"
+#include "Engine/WorldComposition.h"
+
+#if WITH_EDITORONLY_DATA
+
+FFoliageInstanceBaseInfo::FFoliageInstanceBaseInfo()
+	: CachedLocation(FVector::ZeroVector)
+	, CachedRotation(FRotator::ZeroRotator)
+	, CachedDrawScale(1.0f, 1.0f, 1.0f)
+{}
+
+FFoliageInstanceBaseInfo::FFoliageInstanceBaseInfo(UActorComponent* InComponent)
+	: BasePtr(InComponent)
+	, CachedLocation(FVector::ZeroVector)
+	, CachedRotation(FRotator::ZeroRotator)
+	, CachedDrawScale(1.0f, 1.0f, 1.0f)
+{
+	UpdateLocationFromComponent(InComponent);
+}
+
+void FFoliageInstanceBaseInfo::UpdateLocationFromComponent(UActorComponent* InComponent)
+{
+	if (InComponent)
+	{
+		AActor* Owner = Cast<AActor>(InComponent->GetOuter());
+		if (Owner)
+		{
+			const USceneComponent* RootComponent = Owner->GetRootComponent();
+			if (RootComponent)
+			{
+				CachedLocation = RootComponent->RelativeLocation;
+				CachedRotation = RootComponent->RelativeRotation;
+				CachedDrawScale = RootComponent->RelativeScale3D;
+			}
+		}
+	}
+}
+
+FArchive& operator << (FArchive& Ar, FFoliageInstanceBaseInfo& BaseInfo)
+{
+	Ar << BaseInfo.BasePtr;
+	Ar << BaseInfo.CachedLocation;
+	Ar << BaseInfo.CachedRotation;
+	Ar << BaseInfo.CachedDrawScale;
+
+	return Ar;
+}
+
+//////////////////////////////////////////////////////////////////////////
+//
+//////////////////////////////////////////////////////////////////////////
+FFoliageInstanceBaseId FFoliageInstanceBaseCache::InvalidBaseId = INDEX_NONE;
+
+FFoliageInstanceBaseCache::FFoliageInstanceBaseCache()
+	: NextBaseId(1)
+{
+}
+
+FArchive& operator << (FArchive& Ar, FFoliageInstanceBaseCache& InstanceBaseCache)
+{
+	Ar << InstanceBaseCache.NextBaseId;
+	Ar << InstanceBaseCache.InstanceBaseMap;
+	Ar << InstanceBaseCache.InstanceBaseLevelMap;
+
+	if (Ar.IsTransacting())
+	{
+		Ar << InstanceBaseCache.InstanceBaseInvMap;
+	}
+	else if (Ar.IsLoading())
+	{
+		InstanceBaseCache.InstanceBaseInvMap.Empty();
+		for (const auto& Pair : InstanceBaseCache.InstanceBaseMap)
+		{
+			const FFoliageInstanceBaseInfo& BaseInfo = Pair.Value;
+			InstanceBaseCache.InstanceBaseInvMap.Add(BaseInfo.BasePtr, Pair.Key);
+		}
+	}
+	
+	check(InstanceBaseCache.InstanceBaseInvMap.Num() == InstanceBaseCache.InstanceBaseMap.Num());
+
+	return Ar;
+}
+
+FFoliageInstanceBaseId FFoliageInstanceBaseCache::AddInstanceBaseId(UActorComponent* InComponent)
+{
+	FFoliageInstanceBaseId BaseId = FFoliageInstanceBaseCache::InvalidBaseId;
+	if (InComponent && InComponent->CreationMethod != EComponentCreationMethod::ConstructionScript)
+	{
+		BaseId = GetInstanceBaseId(InComponent);
+		if (BaseId == FFoliageInstanceBaseCache::InvalidBaseId)
+		{
+			BaseId = NextBaseId++;
+			FFoliageInstanceBasePtr BasePtr = InComponent;
+
+			InstanceBaseMap.Add(BaseId, FFoliageInstanceBaseInfo(InComponent));
+			InstanceBaseInvMap.Add(BasePtr, BaseId);
+
+			check(InstanceBaseMap.Num() == InstanceBaseInvMap.Num());
+
+			ULevel* ComponentLevel = InComponent->GetComponentLevel();
+			if (ComponentLevel)
+			{
+				UWorld* ComponentWorld = Cast<UWorld>(ComponentLevel->GetOuter());
+				if (ComponentWorld)
+				{
+					auto WorldKey = TAssetPtr<UWorld>(ComponentWorld);
+					InstanceBaseLevelMap.FindOrAdd(WorldKey).Add(BasePtr);
+				}
+			}
+		}
+	}
+
+	return BaseId;
+}
+
+FFoliageInstanceBaseId FFoliageInstanceBaseCache::GetInstanceBaseId(UActorComponent* InComponent) const
+{
+	FFoliageInstanceBasePtr BasePtr = InComponent;
+	return GetInstanceBaseId(BasePtr);
+}
+
+FFoliageInstanceBaseId FFoliageInstanceBaseCache::GetInstanceBaseId(const FFoliageInstanceBasePtr& BasePtr) const
+{
+	const FFoliageInstanceBaseId* BaseId = InstanceBaseInvMap.Find(BasePtr);
+	if (!BaseId)
+	{
+		return FFoliageInstanceBaseCache::InvalidBaseId;
+	}
+	
+	return *BaseId;
+}
+
+FFoliageInstanceBasePtr FFoliageInstanceBaseCache::GetInstanceBasePtr(FFoliageInstanceBaseId BaseId) const
+{
+	return GetInstanceBaseInfo(BaseId).BasePtr;
+}
+
+FFoliageInstanceBaseInfo FFoliageInstanceBaseCache::GetInstanceBaseInfo(FFoliageInstanceBaseId BaseId) const
+{
+	return InstanceBaseMap.FindRef(BaseId);
+}
+
+FFoliageInstanceBaseInfo FFoliageInstanceBaseCache::UpdateInstanceBaseInfoTransform(UActorComponent* InComponent)
+{
+	auto BaseId = GetInstanceBaseId(InComponent);
+	if (BaseId != FFoliageInstanceBaseCache::InvalidBaseId)
+	{
+		auto* BaseInfo = InstanceBaseMap.Find(BaseId);
+		check(BaseInfo);
+		BaseInfo->UpdateLocationFromComponent(InComponent);
+		return *BaseInfo;
+	}
+
+	return FFoliageInstanceBaseInfo();
+}
+
+void FFoliageInstanceBaseCache::CompactInstanceBaseCache(AInstancedFoliageActor* IFA)
+{
+	UWorld* World = IFA->GetWorld();
+	if (!World)
+	{
+		return;
+	}
+
+	TSet<FFoliageInstanceBaseId> BasesInUse;
+	for (auto& Pair : IFA->FoliageMeshes)
+	{
+		for (const auto& Pair : Pair.Value->ComponentHash)
+		{
+			if (Pair.Key != FFoliageInstanceBaseCache::InvalidBaseId)
+			{
+				BasesInUse.Add(Pair.Key);
+			}
+		}
+	}
+	
+	// Look for any removed maps
+	TSet<FFoliageInstanceBasePtr> InvalidBasePtrs;
+	for (auto& Pair : InstanceBaseLevelMap)
+	{
+		const auto& WorldAsset = Pair.Key;
+		
+		bool bExists = (WorldAsset == World);
+		// Check sub-levels
+		if (!bExists)
+		{
+			const FName PackageName = FName(*FPackageName::ObjectPathToPackageName(WorldAsset.ToStringReference().ToString()));
+			if (World->WorldComposition)
+			{
+				bExists = World->WorldComposition->DoesTileExists(PackageName);
+			}
+			else
+			{
+				bExists = (World->GetLevelStreamingForPackageName(PackageName) != nullptr);
+			}
+		}
+
+		if (!bExists)
+		{
+			InvalidBasePtrs.Append(Pair.Value);
+			InstanceBaseLevelMap.Remove(Pair.Key);
+		}
+		else
+		{
+			// Remove dead links
+			for (int32 i = Pair.Value.Num()-1; i >= 0; --i)
+			{
+				// Base needs to be removed if it's not in use by existing instances or component was removed
+				if (Pair.Value[i].IsNull() || !BasesInUse.Contains(GetInstanceBaseId(Pair.Value[i])))
+				{
+					InvalidBasePtrs.Add(Pair.Value[i]);
+					Pair.Value.RemoveAt(i);
+				}
+			}
+
+			if (Pair.Value.Num() == 0)
+			{
+				InstanceBaseLevelMap.Remove(Pair.Key);
+			}
+		}
+	}
+	
+	TSet<FFoliageInstanceBaseId> InvalidBaseIds;
+	InstanceBaseInvMap.Empty();
+	// Look for any removed base components
+	for (const auto& Pair : InstanceBaseMap)
+	{
+		const FFoliageInstanceBaseInfo& BaseInfo = Pair.Value;
+		if (InvalidBasePtrs.Contains(BaseInfo.BasePtr))
+		{
+			InvalidBaseIds.Add(Pair.Key);
+			InstanceBaseMap.Remove(Pair.Key);
+		}
+		else
+		{
+			// Regenerate inverse map
+			check(!InstanceBaseInvMap.Contains(BaseInfo.BasePtr));
+			InstanceBaseInvMap.Add(BaseInfo.BasePtr, Pair.Key);
+		}
+	}
+
+	if (InvalidBaseIds.Num())
+	{
+		for (auto& Pair : IFA->FoliageMeshes)
+		{
+			auto& MeshInfo = Pair.Value;
+			for (FFoliageInstance& Instance : MeshInfo->Instances)
+			{
+				if (InvalidBaseIds.Contains(Instance.BaseId))
+				{
+					Instance.BaseId = FFoliageInstanceBaseCache::InvalidBaseId;
+				}
+			}
+
+			for (auto RemovedBaseId : InvalidBaseIds)
+			{
+				MeshInfo->ComponentHash.Remove(RemovedBaseId);
+			}
+		
+			MeshInfo->ComponentHash.Compact();
+		}
+
+		InstanceBaseMap.Compact();
+		InstanceBaseLevelMap.Compact();
+	}
+}
+
+void FFoliageInstanceBaseCache::UpdateInstanceBaseCachedTransforms()
+{
+	for (auto& Pair : InstanceBaseMap)
+	{
+		FFoliageInstanceBaseInfo& BaseInfo = Pair.Value;
+		BaseInfo.UpdateLocationFromComponent(BaseInfo.BasePtr.Get());
+	}
+}
+
+#endif//WITH_EDITORONLY_DATA
diff --git a/Engine/Source/Runtime/Foliage/Private/InstancedFoliage.cpp b/Engine/Source/Runtime/Foliage/Private/InstancedFoliage.cpp
--- a/Engine/Source/Runtime/Foliage/Private/InstancedFoliage.cpp
+++ b/Engine/Source/Runtime/Foliage/Private/InstancedFoliage.cpp
@@ -40,6 +40,8 @@ struct FFoliageCustomVersion
 		AddedFoliageTypeUpdateGuid = 3,
 		// Use a GUID to determine whic procedural actor spawned us
 		ProceduralGuid = 4,
+		// Support for cross-level bases 
+		CrossLevelBase = 5,
 		// -----<new versions can be added above this line>-------------------------------------------------
 		VersionPlusOne,
 		LatestVersion = VersionPlusOne - 1
@@ -57,11 +59,8 @@ const FGuid FFoliageCustomVersion::GUID(0x430C4D19, 0x71544970, 0x87699B69, 0xDF
 FCustomVersionRegistration GRegisterFoliageCustomVersion(FFoliageCustomVersion::GUID, FFoliageCustomVersion::LatestVersion, TEXT("FoliageVer"));
 
 
-//
-// Serializers for struct data
-//
-
-FArchive& operator<<(FArchive& Ar, FFoliageInstance& Instance)
+// Legacy (< FFoliageCustomVersion::CrossLevelBase) serializer
+FArchive& operator<<(FArchive& Ar, FFoliageInstance_Deprecated& Instance)
 {
 	Ar << Instance.Base;
 	Ar << Instance.Location;
@@ -99,6 +98,52 @@ FArchive& operator<<(FArchive& Ar, FFoliageInstance& Instance)
 	return Ar;
 }
 
+//
+// Serializers for struct data
+//
+FArchive& operator<<(FArchive& Ar, FFoliageInstance& Instance)
+{
+	Ar << Instance.Location;
+	Ar << Instance.Rotation;
+	Ar << Instance.DrawScale3D;
+	Ar << Instance.PreAlignRotation;
+	Ar << Instance.ProceduralGuid;
+	Ar << Instance.Flags;
+	Ar << Instance.ZOffset;
+	Ar << Instance.BaseId;
+
+	return Ar;
+}
+
+static void ConvertDeprecatedFoliageMeshes(
+	AInstancedFoliageActor* IFA, 
+	const TMap<UFoliageType*, TUniqueObj<FFoliageMeshInfo_Deprecated>>& FoliageMeshesDeprecated, 
+	TMap<UFoliageType*, TUniqueObj<FFoliageMeshInfo>>& FoliageMeshes)
+{
+#if WITH_EDITORONLY_DATA	
+	for (auto Pair : FoliageMeshesDeprecated)
+	{
+		auto& FoliageMesh = FoliageMeshes.Add(Pair.Key);
+		const auto& FoliageMeshDeprecated = Pair.Value;
+		
+		FoliageMesh->Component = FoliageMeshDeprecated->Component;
+		FoliageMesh->FoliageTypeUpdateGuid = FoliageMeshDeprecated->FoliageTypeUpdateGuid;
+
+		FoliageMesh->Instances.Reserve(FoliageMeshDeprecated->Instances.Num());
+
+		for (const FFoliageInstance_Deprecated& DeprecatedInstance : FoliageMeshDeprecated->Instances)
+		{
+			FFoliageInstance Instance;
+			static_cast<FFoliageInstancePlacementInfo&>(Instance) = DeprecatedInstance;
+			Instance.BaseId = IFA->InstanceBaseCache.AddInstanceBaseId(DeprecatedInstance.Base);
+			Instance.ProceduralGuid = DeprecatedInstance.ProceduralGuid;
+
+			FoliageMesh->Instances.Add(Instance);
+		}
+	}
+#endif//WITH_EDITORONLY_DATA	
+}
+
 /**
 *	FFoliageInstanceCluster_Deprecated
 */
@@ -130,7 +175,7 @@ struct FFoliageInstanceCluster_Deprecated
 	}
 };
 
-FArchive& operator<<(FArchive& Ar, FFoliageMeshInfo& MeshInfo)
+FArchive& operator<<(FArchive& Ar, FFoliageMeshInfo_Deprecated& MeshInfo)
 {
 	if (Ar.CustomVer(FFoliageCustomVersion::GUID) >= FFoliageCustomVersion::FoliageUsingHierarchicalISMC)
 	{
@@ -153,7 +198,26 @@ FArchive& operator<<(FArchive& Ar, FFoliageMeshInfo& MeshInfo)
 	{
 		Ar << MeshInfo.FoliageTypeUpdateGuid;
 	}
+#endif
+
+	return Ar;
+}
+
+FArchive& operator<<(FArchive& Ar, FFoliageMeshInfo& MeshInfo)
+{
+	Ar << MeshInfo.Component;
+
+#if WITH_EDITORONLY_DATA
+	if (!Ar.ArIsFilterEditorOnly && !(Ar.GetPortFlags() & PPF_DuplicateForPIE))
+	{
+		Ar << MeshInfo.Instances;
+	}
 
+	if (!Ar.ArIsFilterEditorOnly)
+	{
+		Ar << MeshInfo.FoliageTypeUpdateGuid;
+	}
+		
 	// Serialize the transient data for undo.
 	if (Ar.IsTransacting())
 	{
@@ -166,11 +230,6 @@ FArchive& operator<<(FArchive& Ar, FFoliageMeshInfo& MeshInfo)
 	return Ar;
 }
 
-FArchive& operator<<(FArchive& Ar, FFoliageComponentHashInfo& ComponentHashInfo)
-{
-	return Ar << ComponentHashInfo.CachedLocation << ComponentHashInfo.CachedRotation << ComponentHashInfo.CachedDrawScale << ComponentHashInfo.Instances;
-}
-
 //
 // UFoliageType
 //
@@ -372,9 +431,9 @@ void FFoliageMeshInfo::CheckValid()
 	InstanceHash->CheckInstanceCount(InstanceTotal);
 
 	int32 ComponentHashTotal = 0;
-	for( TMap<UActorComponent*, FFoliageComponentHashInfo >::TConstIterator It(ComponentHash); It; ++It )
+	for (const auto& Pair : ComponentHash)
 	{
-		ComponentHashTotal += It.Value().Instances.Num();
+		ComponentHashTotal += Pair.Value().Num();
 	}
 	check( ComponentHashTotal == InstanceTotal);
 
@@ -409,6 +468,13 @@ void FFoliageMeshInfo::CheckValid()
 #endif
 }
 
+void FFoliageMeshInfo::AddInstance(AInstancedFoliageActor* InIFA, const UFoliageType* InSettings, const FFoliageInstance& InNewInstance, UActorComponent* InBaseComponent)
+{
+	FFoliageInstance Instance = InNewInstance;
+	Instance.BaseId = InIFA->InstanceBaseCache.AddInstanceBaseId(InBaseComponent);
+	AddInstance(InIFA, InSettings, Instance);
+}
+
 void FFoliageMeshInfo::AddInstance(AInstancedFoliageActor* InIFA, const UFoliageType* InSettings, const FFoliageInstance& InNewInstance)
 {
 	InIFA->Modify();
@@ -455,20 +521,12 @@ void FFoliageMeshInfo::AddInstance(AInstancedFoliageActor* InIFA, const UFoliage
 	}
 
 	// Add the instance taking either a free slot or adding a new item.
-	int32 InstanceIndex = Instances.AddUninitialized();
-
+	int32 InstanceIndex = Instances.Add(InNewInstance);
 	FFoliageInstance& AddedInstance = Instances[InstanceIndex];
-	AddedInstance = InNewInstance;
 
 	// Add the instance to the hash
 	InstanceHash->InsertInstance(InNewInstance.Location, InstanceIndex);
-	FFoliageComponentHashInfo* ComponentHashInfo = ComponentHash.Find(InNewInstance.Base);
-	if (ComponentHashInfo == nullptr)
-	{
-		ComponentHashInfo = &ComponentHash.Add(InNewInstance.Base, FFoliageComponentHashInfo(InNewInstance.Base));
-	}
-	ComponentHashInfo->Instances.Add(InstanceIndex);
-	
+	ComponentHash.FindOrAdd(InNewInstance.BaseId).Add(InstanceIndex);
 	// Calculate transform for the instance
 	FTransform InstanceToWorld = InNewInstance.GetInstanceWorldTransform();
 
@@ -508,16 +566,7 @@ void FFoliageMeshInfo::RemoveInstances(AInstancedFoliageActor* InIFA, const TArr
 
 			// remove from hash
 			InstanceHash->RemoveInstance(Instance.Location, InstanceIndex);
-		    FFoliageComponentHashInfo* ComponentHashInfo = ComponentHash.Find(Instance.Base);
-		    if (ComponentHashInfo)
-		    {
-			    ComponentHashInfo->Instances.Remove(InstanceIndex);
-			    if (ComponentHashInfo->Instances.Num() == 0)
-			    {
-				    // Remove the component from the component hash if this is the last instance.
-				    ComponentHash.Remove(Instance.Base);
-			    }
-		    }
+			SetInstanceBaseId(FFoliageInstanceBaseCache::InvalidBaseId, InstanceIndex);
 
 			// remove from the component
 			Component->RemoveInstance(InstanceIndex);
@@ -537,12 +586,11 @@ void FFoliageMeshInfo::RemoveInstances(AInstancedFoliageActor* InIFA, const TArr
 				InstanceHash->InsertInstance(SwappedInstance.Location, InstanceIndex);
 
 				// Component hash
-				if (SwappedInstance.Base)
+				auto* InstanceSet = ComponentHash.Find(SwappedInstance.BaseId);
+				if (InstanceSet)
 				{
-					FFoliageComponentHashInfo* ComponentHashInfo = ComponentHash.Find(SwappedInstance.Base);
-					check(ComponentHashInfo);
-					ComponentHashInfo->Instances.Remove(Instances.Num());
-					ComponentHashInfo->Instances.Add(InstanceIndex);
+					InstanceSet->Remove(Instances.Num());
+					InstanceSet->Add(InstanceIndex);
 				}
 
 				// Selection
@@ -630,6 +678,34 @@ int32 FFoliageMeshInfo::GetInstanceCount() const
 	return Instances.Num();
 }
 
+void FFoliageMeshInfo::SetInstanceBaseId(FFoliageInstanceBaseId NewBaseId, int32 InstanceIndex)
+{
+	FFoliageInstance& Instance = Instances[InstanceIndex];
+	if (NewBaseId != Instance.BaseId)
+	{
+		// Remove current base link
+		auto* InstanceSet = ComponentHash.Find(Instance.BaseId);
+		if (InstanceSet)
+		{
+			InstanceSet->Remove(InstanceIndex);
+			if (InstanceSet->Num() == 0)
+			{
+				// Remove the component from the component hash if this is the last instance.
+				ComponentHash.Remove(Instance.BaseId);
+			}
+		}
+
+		// Assign new
+		Instance.BaseId = NewBaseId;
+		
+		if (NewBaseId != FFoliageInstanceBaseCache::InvalidBaseId)
+		{
+			ComponentHash.FindOrAdd(NewBaseId).Add(InstanceIndex);
+		}
+	}
+}
+
+
 // Destroy existing clusters and reassign all instances to new clusters
 void FFoliageMeshInfo::ReallocateClusters(AInstancedFoliageActor* InIFA, UFoliageType* InSettings)
 {
@@ -920,45 +996,52 @@ const FFoliageMeshInfo* AInstancedFoliageActor::FindMesh(const UFoliageType* InT
 #if WITH_EDITOR
 void AInstancedFoliageActor::MoveInstancesForMovedComponent(UActorComponent* InComponent)
 {
+	const auto BaseId = InstanceBaseCache.GetInstanceBaseId(InComponent);
+	if (BaseId == FFoliageInstanceBaseCache::InvalidBaseId)
+	{
+		return;
+	}
+
 	bool bUpdatedInstances = false;
 	bool bFirst = true;
 
+	const auto OldBaseInfo = InstanceBaseCache.GetInstanceBaseInfo(BaseId);
+	const auto NewBaseInfo = InstanceBaseCache.UpdateInstanceBaseInfoTransform(InComponent);
+
+	FMatrix DeltaTransfrom = 
+		FTranslationMatrix(-OldBaseInfo.CachedLocation) *
+		FInverseRotationMatrix(OldBaseInfo.CachedRotation) *
+		FScaleMatrix(NewBaseInfo.CachedDrawScale / OldBaseInfo.CachedDrawScale) *
+		FRotationMatrix(NewBaseInfo.CachedRotation) *
+		FTranslationMatrix(NewBaseInfo.CachedLocation);
+
 	for (auto& MeshPair : FoliageMeshes)
 	{
 		FFoliageMeshInfo& MeshInfo = *MeshPair.Value;
-		FFoliageComponentHashInfo* ComponentHashInfo = MeshInfo.ComponentHash.Find(InComponent);
-		if (ComponentHashInfo)
+		const auto* InstanceSet = MeshInfo.ComponentHash.Find(BaseId);
+		if (InstanceSet && InstanceSet->Num())
 		{
 			if (bFirst)
 			{
 				bFirst = false;
 				Modify();
 			}
-
-			FVector OldLocation = ComponentHashInfo->CachedLocation;
-			FRotator OldRotation = ComponentHashInfo->CachedRotation;
-			FVector OldDrawScale = ComponentHashInfo->CachedDrawScale;
-			ComponentHashInfo->UpdateLocationFromActor(InComponent);
-
-			for (int32 InstanceIndex : ComponentHashInfo->Instances)
+					
+			for (int32 InstanceIndex : *InstanceSet)
 			{
 				FFoliageInstance& Instance = MeshInfo.Instances[InstanceIndex];
 
 				MeshInfo.InstanceHash->RemoveInstance(Instance.Location, InstanceIndex);
 
 				// Apply change
-				FMatrix DeltaTransform =
+				FMatrix NewTransform =
 					FRotationMatrix(Instance.Rotation) *
 					FTranslationMatrix(Instance.Location) *
-					FTranslationMatrix(-OldLocation) *
-					FInverseRotationMatrix(OldRotation) *
-					FScaleMatrix(ComponentHashInfo->CachedDrawScale / OldDrawScale) *
-					FRotationMatrix(ComponentHashInfo->CachedRotation) *
-					FTranslationMatrix(ComponentHashInfo->CachedLocation);
+					DeltaTransfrom;
 
 				// Extract rotation and position
-				Instance.Location = DeltaTransform.GetOrigin();
-				Instance.Rotation = DeltaTransform.Rotator();
+				Instance.Location = NewTransform.GetOrigin();
+				Instance.Rotation = NewTransform.Rotator();
 
 				// Apply render data
 				check(MeshInfo.Component);
@@ -976,10 +1059,11 @@ void AInstancedFoliageActor::DeleteInstancesForComponent(UActorComponent* InComp
 	for (auto& MeshPair : FoliageMeshes)
 	{
 		FFoliageMeshInfo& MeshInfo = *MeshPair.Value;
-		const FFoliageComponentHashInfo* ComponentHashInfo = MeshInfo.ComponentHash.Find(InComponent);
-		if (ComponentHashInfo)
+		const auto BaseId = InstanceBaseCache.GetInstanceBaseId(InComponent);
+		const auto* InstanceSet = MeshInfo.ComponentHash.Find(BaseId);
+		if (InstanceSet)
 		{
-			MeshInfo.RemoveInstances(this, ComponentHashInfo->Instances.Array());
+			MeshInfo.RemoveInstances(this, InstanceSet->Array());
 		}
 	}
 }
@@ -1009,7 +1093,8 @@ void AInstancedFoliageActor::DeleteInstancesForProceduralFoliageComponent(const
 void AInstancedFoliageActor::MoveInstancesForComponentToCurrentLevel(UActorComponent* InComponent)
 {
 	AInstancedFoliageActor* NewIFA = AInstancedFoliageActor::GetInstancedFoliageActorForCurrentLevel(InComponent->GetWorld());
-
+	const auto SourceBaseId = InstanceBaseCache.GetInstanceBaseId(InComponent);
+	
 	for (auto& MeshPair : FoliageMeshes)
 	{
 		FFoliageMeshInfo& MeshInfo = *MeshPair.Value;
@@ -1021,19 +1106,19 @@ void AInstancedFoliageActor::MoveInstancesForComponentToCurrentLevel(UActorCompo
 			FoliageType = (UFoliageType*)StaticDuplicateObject(FoliageType, NewIFA, nullptr, RF_AllFlags & ~(RF_Standalone | RF_Public));
 		}
 		
-		const FFoliageComponentHashInfo* ComponentHashInfo = MeshInfo.ComponentHash.Find(InComponent);
-		if (ComponentHashInfo)
+		const auto* InstanceSet = MeshInfo.ComponentHash.Find(SourceBaseId);
+		if (InstanceSet)
 		{
 			FFoliageMeshInfo* NewMeshInfo = NewIFA->FindOrAddMesh(FoliageType);
 
 			// Add the foliage to the new level
-			for (int32 InstanceIndex : ComponentHashInfo->Instances)
+			for (int32 InstanceIndex : *InstanceSet)
 			{
-				NewMeshInfo->AddInstance(NewIFA, FoliageType, MeshInfo.Instances[InstanceIndex]);
+				NewMeshInfo->AddInstance(NewIFA, FoliageType, MeshInfo.Instances[InstanceIndex], InComponent);
 			}
 
 			// Remove from old level
-			MeshInfo.RemoveInstances(this, ComponentHashInfo->Instances.Array());
+			MeshInfo.RemoveInstances(this, InstanceSet->Array());
 		}
 	}
 }
@@ -1047,6 +1132,9 @@ void AInstancedFoliageActor::MoveInstancesToNewComponent(UPrimitiveComponent* In
 	}
 	check(NewIFA);
 
+	const auto OldBaseId = this->InstanceBaseCache.GetInstanceBaseId(InOldComponent);
+	const auto NewBaseId = NewIFA->InstanceBaseCache.AddInstanceBaseId(InNewComponent);
+
 	for (auto& MeshPair : FoliageMeshes)
 	{
 		FFoliageMeshInfo& MeshInfo = *MeshPair.Value;
@@ -1058,40 +1146,36 @@ void AInstancedFoliageActor::MoveInstancesToNewComponent(UPrimitiveComponent* In
 			FoliageType = (UFoliageType*)StaticDuplicateObject(FoliageType, NewIFA, nullptr, RF_AllFlags & ~(RF_Standalone | RF_Public));
 		}
 
-		const FFoliageComponentHashInfo* ComponentHashInfo = MeshInfo.ComponentHash.Find(InOldComponent);
-		if (ComponentHashInfo)
+		TSet<int32> InstanceSet;
+		if (MeshInfo.ComponentHash.RemoveAndCopyValue(OldBaseId, InstanceSet) && InstanceSet.Num())
 		{
-			// For same level can just remap the instances, otherwise we have to do a more complex move
+			// For same FoliageActor can just remap the instances, otherwise we have to do a more complex move
 			if (NewIFA == this)
 			{
-				FFoliageComponentHashInfo NewComponentHashInfo = MoveTemp(*ComponentHashInfo);
-				NewComponentHashInfo.UpdateLocationFromActor(InNewComponent);
-
 				// Update the instances
-				for (int32 InstanceIndex : NewComponentHashInfo.Instances)
+				for (int32 InstanceIndex : InstanceSet)
 				{
-					MeshInfo.Instances[InstanceIndex].Base = InNewComponent;
+					MeshInfo.Instances[InstanceIndex].BaseId = NewBaseId;
 				}
-
+				
 				// Update the hash
-				MeshInfo.ComponentHash.Add(InNewComponent, MoveTemp(NewComponentHashInfo));
-				MeshInfo.ComponentHash.Remove(InOldComponent);
+				MeshInfo.ComponentHash.Add(NewBaseId, MoveTemp(InstanceSet));
+				
 			}
 			else
 			{
 				FFoliageMeshInfo* NewMeshInfo = NewIFA->FindOrAddMesh(FoliageType);
 
 				// Add the foliage to the new level
-				for (int32 InstanceIndex : ComponentHashInfo->Instances)
+				for (int32 InstanceIndex : InstanceSet)
 				{
 					FFoliageInstance NewInstance = MeshInfo.Instances[InstanceIndex];
-					NewInstance.Base = InNewComponent;
+					NewInstance.BaseId = NewBaseId;
 					NewMeshInfo->AddInstance(NewIFA, FoliageType, NewInstance);
 				}
 
 				// Remove from old level
-				MeshInfo.RemoveInstances(this, ComponentHashInfo->Instances.Array());
-				check(!MeshInfo.ComponentHash.Contains(InOldComponent));
+				MeshInfo.RemoveInstances(this, InstanceSet.Array());
 			}
 		}
 	}
@@ -1100,18 +1184,18 @@ void AInstancedFoliageActor::MoveInstancesToNewComponent(UPrimitiveComponent* In
 TMap<UFoliageType*, TArray<const FFoliageInstancePlacementInfo*>> AInstancedFoliageActor::GetInstancesForComponent(UActorComponent* InComponent)
 {
 	TMap<UFoliageType*, TArray<const FFoliageInstancePlacementInfo*>> Result;
+	const auto BaseId = InstanceBaseCache.GetInstanceBaseId(InComponent);
 
 	for (auto& MeshPair : FoliageMeshes)
 	{
 		const FFoliageMeshInfo& MeshInfo = *MeshPair.Value;
-
-		const FFoliageComponentHashInfo* ComponentHashInfo = MeshInfo.ComponentHash.Find(InComponent);
-		if (ComponentHashInfo)
+		const auto* InstanceSet = MeshInfo.ComponentHash.Find(BaseId);
+		if (InstanceSet)
 		{
 			TArray<const FFoliageInstancePlacementInfo*>& Array = Result.Add(MeshPair.Key, TArray<const FFoliageInstancePlacementInfo*>());
-			Array.Empty(ComponentHashInfo->Instances.Num());
+			Array.Empty(InstanceSet->Num());
 
-			for (int32 InstanceIndex : ComponentHashInfo->Instances)
+			for (int32 InstanceIndex : *InstanceSet)
 			{
 				const FFoliageInstancePlacementInfo* Instance = &MeshInfo.Instances[InstanceIndex];
 				Array.Add(Instance);
@@ -1426,17 +1510,17 @@ void AInstancedFoliageActor::MapRebuild()
 		for (auto& ComponentFoliagePair : MeshInfo.ComponentHash)
 		{
 			// BSP components are UModelComponents - they are the only ones we need to change
-			UModelComponent* TargetComponent = Cast<UModelComponent>(ComponentFoliagePair.Key);
+			auto BaseComponentPtr = InstanceBaseCache.GetInstanceBasePtr(ComponentFoliagePair.Key);
+			UModelComponent* TargetComponent = Cast<UModelComponent>(BaseComponentPtr.Get());
+		
 			// Check if it's part of a brush. We only need to fix up model components that are part of the level BSP.
 			if (TargetComponent && Cast<ABrush>(TargetComponent->GetOuter()) == nullptr)
 			{
 				// Delete its instances later
 				RemovedModelComponents.Add(TargetComponent);
 
-				FFoliageComponentHashInfo const& FoliageInfo = ComponentFoliagePair.Value;
-
 				// We have to test each instance to see if we can migrate it across
-				for (int32 InstanceIdx : FoliageInfo.Instances)
+				for (int32 InstanceIdx : ComponentFoliagePair.Value)
 				{
 					// Use a line test against the world. This is not very reliable as we don't know the original trace direction.
 					check(MeshInfo.Instances.IsValidIndex(InstanceIdx));
@@ -1454,7 +1538,7 @@ void AInstancedFoliageActor::MapRebuild()
 
 					if (bHit && Result.Component.IsValid() && Result.Component->IsA(UModelComponent::StaticClass()))
 					{
-						NewInstance.Base = CastChecked<UPrimitiveComponent>(Result.Component.Get());
+						NewInstance.BaseId = InstanceBaseCache.AddInstanceBaseId(Result.Component.Get());
 						NewInstances.FindOrAdd(Settings).Add(NewInstance);
 					}
 				}
@@ -1481,45 +1565,12 @@ void AInstancedFoliageActor::MapRebuild()
 	}
 }
 
-void AInstancedFoliageActor::ApplyLevelTransform(const FTransform& LevelTransform)
-{
-	// Apply transform to foliage editor only data
-	if (GIsEditor)
-	{
-		for (auto& MeshPair : FoliageMeshes)
-		{
-			FFoliageMeshInfo& MeshInfo = *MeshPair.Value;
-		
-			MeshInfo.InstanceHash->Empty();
-			for (int32 InstanceIdx = 0; InstanceIdx < MeshInfo.Instances.Num(); InstanceIdx++)
-			{
-				FFoliageInstance& Instance = MeshInfo.Instances[InstanceIdx];
-				FTransform NewTransform = Instance.GetInstanceWorldTransform() * LevelTransform;
-			
-				Instance.Location		= NewTransform.GetLocation();
-				Instance.Rotation		= NewTransform.GetRotation().Rotator();
-				Instance.DrawScale3D	= NewTransform.GetScale3D();
-			
-				// Rehash instance location
-				MeshInfo.InstanceHash->InsertInstance(Instance.Location, InstanceIdx);
-			}
-
-			for (auto It = MeshInfo.ComponentHash.CreateIterator(); It; ++It)
-			{
-				// We assume here that component we painted foliage on, was transformed as well
-				FFoliageComponentHashInfo& Info = It.Value();
-				Info.UpdateLocationFromActor(It.Key());
-			}
-		}
-	}
-}
-
 #endif // WITH_EDITOR
 
 struct FFoliageMeshInfo_Old
 {
 	TArray<FFoliageInstanceCluster_Deprecated> InstanceClusters;
-	TArray<FFoliageInstance> Instances;
+	TArray<FFoliageInstance_Deprecated> Instances;
 	UFoliageType_InstancedStaticMesh* Settings; // Type remapped via +ActiveClassRedirects
 };
 FArchive& operator<<(FArchive& Ar, FFoliageMeshInfo_Old& MeshInfo)
@@ -1537,13 +1588,21 @@ void AInstancedFoliageActor::Serialize(FArchive& Ar)
 
 	Ar.UsingCustomVersion(FFoliageCustomVersion::GUID);
 
+#if WITH_EDITORONLY_DATA
+	if (!Ar.ArIsFilterEditorOnly && Ar.CustomVer(FFoliageCustomVersion::GUID) >= FFoliageCustomVersion::CrossLevelBase)
+	{
+		Ar << InstanceBaseCache;
+	}
+#endif
+	
 	if (Ar.UE4Ver() < VER_UE4_FOLIAGE_SETTINGS_TYPE)
 	{
+		TMap<UFoliageType*, TUniqueObj<FFoliageMeshInfo_Deprecated>> FoliageMeshesDeprecated;
 		TMap<UStaticMesh*, FFoliageMeshInfo_Old> OldFoliageMeshes;
 		Ar << OldFoliageMeshes;
 		for (auto& OldMeshInfo : OldFoliageMeshes)
 		{
-			FFoliageMeshInfo NewMeshInfo;
+			FFoliageMeshInfo_Deprecated NewMeshInfo;
 
 #if WITH_EDITORONLY_DATA
 			NewMeshInfo.Instances = MoveTemp(OldMeshInfo.Value.Instances);
@@ -1570,14 +1629,25 @@ void AInstancedFoliageActor::Serialize(FArchive& Ar)
 #if WITH_EDITORONLY_DATA
 			NewMeshInfo.FoliageTypeUpdateGuid = FoliageType->UpdateGuid;
 #endif
-			FoliageMeshes.Add(FoliageType, TUniqueObj<FFoliageMeshInfo>(MoveTemp(NewMeshInfo)));
+			FoliageMeshesDeprecated.Add(FoliageType, TUniqueObj<FFoliageMeshInfo_Deprecated>(MoveTemp(NewMeshInfo)));
 		}
+
+		ConvertDeprecatedFoliageMeshes(this, FoliageMeshesDeprecated, FoliageMeshes);
 	}
 	else
 	{
-		Ar << FoliageMeshes;
+		if (Ar.CustomVer(FFoliageCustomVersion::GUID) < FFoliageCustomVersion::CrossLevelBase)
+		{
+			TMap<UFoliageType*, TUniqueObj<FFoliageMeshInfo_Deprecated>> FoliageMeshesDeprecated;
+			Ar << FoliageMeshesDeprecated;
+			ConvertDeprecatedFoliageMeshes(this, FoliageMeshesDeprecated, FoliageMeshes);
+		}
+		else
+		{
+			Ar << FoliageMeshes;
+		}
 	}
-
+	
 	// Clean up any old cluster components and convert to hierarchical instanced foliage.
 	if (Ar.CustomVer(FFoliageCustomVersion::GUID) < FFoliageCustomVersion::FoliageUsingHierarchicalISMC)
 	{
@@ -1590,28 +1660,38 @@ void AInstancedFoliageActor::Serialize(FArchive& Ar)
 	}
 }
 
-void AInstancedFoliageActor::PostLoad()
+void AInstancedFoliageActor::PostInitProperties()
 {
-	Super::PostLoad();
+	Super::PostInitProperties();
 
-#if WITH_EDITORONLY_DATA
-	if (GetLinkerUE4Version() < VER_UE4_DISALLOW_FOLIAGE_ON_BLUEPRINTS)
+#if WITH_EDITOR
+	if (!IsTemplate())
 	{
-		for (auto& MeshPair : FoliageMeshes)
-		{
-			for (FFoliageInstance& Instance : MeshPair.Value->Instances)
-			{
-				// Clear out the Base for any instances based on blueprint-created components,
-				// as those components will be destroyed when the construction scripts are
-				// re-run, leaving dangling references and causing crashes (woo!)
-				if (Instance.Base && Instance.Base->CreationMethod == EComponentCreationMethod::ConstructionScript)
-				{
-					Instance.Base = NULL;
-				}
-			}
-		}
+		GEngine->OnActorMoved().Remove(OnLevelActorMovedDelegateHandle);
+		OnLevelActorMovedDelegateHandle = GEngine->OnActorMoved().AddUObject(this, &AInstancedFoliageActor::OnLevelActorMoved);
+
+		GEngine->OnLevelActorDeleted().Remove(OnLevelActorDeletedDelegateHandle);
+		OnLevelActorDeletedDelegateHandle = GEngine->OnLevelActorDeleted().AddUObject(this, &AInstancedFoliageActor::OnLevelActorDeleted);
+	}
+#endif
+}
+
+void AInstancedFoliageActor::BeginDestroy()
+{
+	Super::BeginDestroy();
+
+#if WITH_EDITOR
+	if (!IsTemplate())
+	{
+		GEngine->OnActorMoved().Remove(OnLevelActorMovedDelegateHandle);
+		GEngine->OnLevelActorDeleted().Remove(OnLevelActorDeletedDelegateHandle);
 	}
 #endif
+}
+
+void AInstancedFoliageActor::PostLoad()
+{
+	Super::PostLoad();
 
 #if WITH_EDITOR
 	if (GIsEditor)
@@ -1668,21 +1748,37 @@ void AInstancedFoliageActor::PostLoad()
 			}
 
 			// Update the hash.
+			MeshInfo.ComponentHash.Empty();
+			MeshInfo.InstanceHash->Empty();
 			for (int32 InstanceIdx = 0; InstanceIdx < MeshInfo.Instances.Num(); InstanceIdx++)
 			{
 				FFoliageInstance& Instance = MeshInfo.Instances[InstanceIdx];
 				{
 					// Add valid instances to the hash.
 					MeshInfo.InstanceHash->InsertInstance(Instance.Location, InstanceIdx);
-					FFoliageComponentHashInfo* ComponentHashInfo = MeshInfo.ComponentHash.Find(Instance.Base);
-					if (ComponentHashInfo == nullptr)
-					{
-						ComponentHashInfo = &MeshInfo.ComponentHash.Add(Instance.Base, FFoliageComponentHashInfo(Instance.Base));
-					}
-					ComponentHashInfo->Instances.Add(InstanceIdx);
+					MeshInfo.ComponentHash.FindOrAdd(Instance.BaseId).Add(InstanceIdx);
+				}
+			}
+			
+			// Fix-up cross level references after serializing deprecated data
+			if (GetLinkerCustomVersion(FFoliageCustomVersion::GUID) < FFoliageCustomVersion::CrossLevelBase)
+			{
+				// update base cached transform this information was not available during serialization
+				InstanceBaseCache.UpdateInstanceBaseCachedTransforms();
+				// there were no cross-level references before
+				check(InstanceBaseCache.InstanceBaseLevelMap.Num() <= 1); 
+				// populate WorldAsset->BasePtr map
+				InstanceBaseCache.InstanceBaseLevelMap.Empty();
+				auto& BaseList = InstanceBaseCache.InstanceBaseLevelMap.Add(TAssetPtr<UWorld>(Cast<UWorld>(GetLevel()->GetOuter())));
+				for (auto& BaseInfoPair : InstanceBaseCache.InstanceBaseMap)
+				{
+					BaseList.Add(BaseInfoPair.Value.BasePtr);
 				}
 			}
 
+			// Clean up dead cross-level references
+			InstanceBaseCache.CompactInstanceBaseCache(this);
+			
 			// Convert to Heirarchical foliage
 			if (GetLinkerCustomVersion(FFoliageCustomVersion::GUID) < FFoliageCustomVersion::FoliageUsingHierarchicalISMC)
 			{
@@ -1701,6 +1797,18 @@ void AInstancedFoliageActor::PostLoad()
 #endif
 }
 
+void AInstancedFoliageActor::PreSave()
+{
+	Super::PreSave();
+
+#if WITH_EDITOR
+	if (GIsEditor)
+	{
+		InstanceBaseCache.CompactInstanceBaseCache(this);
+	}
+#endif //WITH_EDITOR
+}
+
 #if WITH_EDITOR
 
 void AInstancedFoliageActor::NotifyFoliageTypeChanged(UFoliageType* FoliageType)
@@ -1714,23 +1822,36 @@ void AInstancedFoliageActor::NotifyFoliageTypeChanged(UFoliageType* FoliageType)
 
 void AInstancedFoliageActor::OnLevelActorMoved(AActor* InActor)
 {
-	if (!InActor->GetWorld()->IsPlayInEditor())
+	UWorld* InWorld = InActor->GetWorld();
+	
+	if (!InWorld || !InWorld->IsGameWorld())
 	{
-		if(AInstancedFoliageActor* IFA = AInstancedFoliageActor::GetInstancedFoliageActorForLevel(InActor->GetLevel()))
+		TInlineComponentArray<UActorComponent*> Components;
+		InActor->GetComponents(Components);
+		
+		for (auto Component : Components)
 		{
-			if (IFA == this)	//IFA for each level is registered so we'll get hit multiple times for each sublevel. @TODO: Should we only register the persistant level IFA?
-			{
-				TInlineComponentArray<UActorComponent*> Components;
-				InActor->GetComponents(Components);
+			MoveInstancesForMovedComponent(Component);
+		}
+	}
+}
 
-				for (int32 Idx = 0; Idx < Components.Num(); ++Idx)
-				{
-					IFA->MoveInstancesForMovedComponent(Components[Idx]);
-				}
-			}
+void AInstancedFoliageActor::OnLevelActorDeleted(AActor* InActor)
+{
+	UWorld* InWorld = InActor->GetWorld();
+	
+	if (!InWorld || !InWorld->IsGameWorld())
+	{
+		TInlineComponentArray<UActorComponent*> Components;
+		InActor->GetComponents(Components);
+		
+		for (auto Component : Components)
+		{
+			DeleteInstancesForComponent(Component);
 		}
 	}
 }
+
 #endif
 
 //
@@ -1749,16 +1870,6 @@ void AInstancedFoliageActor::AddReferencedObjects(UObject* InThis, FReferenceCol
 		Collector.AddReferencedObject(MeshPair.Key, This);
 		FFoliageMeshInfo& MeshInfo = *MeshPair.Value;
 
-#if WITH_EDITORONLY_DATA
-		for (FFoliageInstance& Instance : MeshInfo.Instances)
-		{
-			if (Instance.Base != nullptr)
-			{
-				Collector.AddReferencedObject(Instance.Base, This);
-			}
-		}
-#endif
-
 		if (MeshInfo.Component)
 		{
 			Collector.AddReferencedObject(MeshInfo.Component, This);
@@ -1767,39 +1878,6 @@ void AInstancedFoliageActor::AddReferencedObjects(UObject* InThis, FReferenceCol
 	Super::AddReferencedObjects(This, Collector);
 }
 
-void AInstancedFoliageActor::PostRegisterAllComponents()
-{
-	Super::PostRegisterAllComponents();
-
-#if WITH_EDITOR
-	if (!IsTemplate())
-	{
-		ULevel* Level = GetLevel();
-		
-		Level->OnApplyLevelTransform.Remove(OnApplyLevelTransformDelegateHandle); 
-		OnApplyLevelTransformDelegateHandle = Level->OnApplyLevelTransform.AddUObject(this, &AInstancedFoliageActor::ApplyLevelTransform);
-
-		GEngine->OnActorMoved().Remove(OnLevelActorMovedDelegateHandle);
-		OnLevelActorMovedDelegateHandle = GEngine->OnActorMoved().AddUObject(this, &AInstancedFoliageActor::OnLevelActorMoved);
-
-	}
-#endif// WITH_EDITOR
-}
-
-void AInstancedFoliageActor::PostUnregisterAllComponents()
-{
-#if WITH_EDITOR
-	if (!IsTemplate())
-	{
-		ULevel* Level = GetLevel();
-		Level->OnApplyLevelTransform.Remove(OnApplyLevelTransformDelegateHandle);
-		GEngine->OnActorMoved().Remove(OnLevelActorMovedDelegateHandle);
-	}
-#endif// WITH_EDITOR
-
-	Super::PostUnregisterAllComponents();
-}
-
 void AInstancedFoliageActor::ApplyWorldOffset(const FVector& InOffset, bool bWorldShift)
 {
 	Super::ApplyWorldOffset(InOffset, bWorldShift);
@@ -1811,6 +1889,8 @@ void AInstancedFoliageActor::ApplyWorldOffset(const FVector& InOffset, bool bWor
 			FFoliageMeshInfo& MeshInfo = *MeshPair.Value;
 
 #if WITH_EDITORONLY_DATA
+			InstanceBaseCache.UpdateInstanceBaseCachedTransforms();
+			
 			MeshInfo.InstanceHash->Empty();
 			for (int32 InstanceIdx = 0; InstanceIdx < MeshInfo.Instances.Num(); InstanceIdx++)
 			{
@@ -1819,19 +1899,12 @@ void AInstancedFoliageActor::ApplyWorldOffset(const FVector& InOffset, bool bWor
 				// Rehash instance location
 				MeshInfo.InstanceHash->InsertInstance(Instance.Location, InstanceIdx);
 			}
-
-			for (auto It = MeshInfo.ComponentHash.CreateIterator(); It; ++It)
-			{
-				// We assume here that component we painted foliage on will be shifted by same value
-				FFoliageComponentHashInfo& Info = It.Value();
-				Info.CachedLocation += InOffset;
-			}
 #endif
 		}
 	}
 }
 
-AInstancedFoliageActor* AInstancedFoliageActor::FoliageTrace(UWorld* InWorld, FHitResult& OutHit, const FDesiredFoliageInstance& DesiredInstance, const AInstancedFoliageActor* IgnoreIFA, FName InTraceTag, bool InbReturnFaceIndex)
+bool AInstancedFoliageActor::FoliageTrace(UWorld* InWorld, FHitResult& OutHit, const FDesiredFoliageInstance& DesiredInstance, const AInstancedFoliageActor* IgnoreIFA, FName InTraceTag, bool InbReturnFaceIndex)
 {
 	FCollisionQueryParams QueryParams(InTraceTag, true, IgnoreIFA);
 	QueryParams.bReturnFaceIndex = InbReturnFaceIndex;
@@ -1854,12 +1927,12 @@ AInstancedFoliageActor* AInstancedFoliageActor::FoliageTrace(UWorld* InWorld, FH
 					const AProceduralFoliageActor* ProceduralFoliageActor = ProceduralFoliageBlockingVolume->ProceduralFoliageActor;
 					if (ProceduralFoliageActor == nullptr || ProceduralFoliageActor->ProceduralComponent == nullptr || ProceduralFoliageActor->ProceduralComponent->GetProceduralGuid() == DesiredInstance.ProceduralGuid)
 					{
-						return nullptr;
+						return false;
 					}
-
-				}else if(Cast<AInstancedFoliageActor>(Hit.Actor.Get()))
+				}
+				else if(Cast<AInstancedFoliageActor>(Hit.Actor.Get()))
 				{
-					return nullptr;
+					return false;
 				}
 			}
 		}
@@ -1871,17 +1944,14 @@ AInstancedFoliageActor* AInstancedFoliageActor::FoliageTrace(UWorld* InWorld, FH
 			continue;
 		}
 
-		if (Hit.Component.IsValid())
+		if (Hit.Component.IsValid() && Hit.Component->GetComponentLevel())
 		{
-			if (ULevel* OwningLevel = Hit.Component->GetComponentLevel())
-			{
-				OutHit = Hit;
-				return AInstancedFoliageActor::GetInstancedFoliageActorForLevel(OwningLevel);
-			}
+			OutHit = Hit;
+			return true;
 		}
 	}
 
-	return nullptr;
+	return false;
 }
 
 bool AInstancedFoliageActor::CheckCollisionWithWorld(const UFoliageType* Settings, const FFoliageInstance& Inst, const FVector& HitNormal, const FVector& HitLocation)
@@ -2017,7 +2087,9 @@ bool FPotentialInstance::PlaceInstance(const UFoliageType* Settings, FFoliageIns
 		}
 	}
 
-	Inst.Base = HitComponent;
+#if WITH_EDITORONLY_DATA
+	Inst.BaseId = IFA->InstanceBaseCache.AddInstanceBaseId(HitComponent);
+#endif//WITH_EDITORONLY_DATA
 
 	return bSkipCollision || IFA->CheckCollisionWithWorld(Settings, Inst, HitNormal, HitLocation);
 }
diff --git a/Engine/Source/Runtime/Foliage/Public/FoliageInstanceBase.h b/Engine/Source/Runtime/Foliage/Public/FoliageInstanceBase.h
new file mode 100644
--- /dev/null
+++ b/Engine/Source/Runtime/Foliage/Public/FoliageInstanceBase.h
@@ -0,0 +1,83 @@
+// Copyright 1998-2015 Epic Games, Inc. All Rights Reserved.
+#pragma once
+
+class AInstancedFoliageActor;
+
+typedef int32 FFoliageInstanceBaseId;
+typedef TLazyObjectPtr<UActorComponent> FFoliageInstanceBasePtr;
+
+#if WITH_EDITORONLY_DATA
+/**
+ * FFoliageInstanceBaseInfo
+ * Cached transform information about base component
+ * Used for adjusting instance transforms after operations on base components with foliage painted on them.
+ */
+struct FFoliageInstanceBaseInfo
+{
+	// tors
+	FFoliageInstanceBaseInfo();
+	FFoliageInstanceBaseInfo(UActorComponent* InComponent);
+
+	// Cache the location and rotation from the actor
+	void UpdateLocationFromComponent(UActorComponent* InComponent);
+
+	// Serializer
+	friend FArchive& operator<<(FArchive& Ar, FFoliageInstanceBaseInfo& ComponentInfo);
+
+	FFoliageInstanceBasePtr BasePtr;
+	FVector		CachedLocation;
+	FRotator	CachedRotation;
+	FVector		CachedDrawScale;
+};
+
+struct FFoliageInstanceBaseCache
+{
+	FFoliageInstanceBaseCache();
+
+	/* Adds new base to cache and/or returns existing base Id*/
+	FOLIAGE_API FFoliageInstanceBaseId AddInstanceBaseId(UActorComponent* InComponent);
+	
+	/* Returns Id for a registered base component, invalid Id otherwise */
+	FOLIAGE_API FFoliageInstanceBaseId GetInstanceBaseId(UActorComponent* InComponent) const;
+	
+	/* Returns Id for a registered base component, invalid Id otherwise */
+	FOLIAGE_API FFoliageInstanceBaseId GetInstanceBaseId(const FFoliageInstanceBasePtr& BasePtr) const;
+	
+	/* Returns registered base component for specified Id, null otherwise */
+	FOLIAGE_API FFoliageInstanceBasePtr GetInstanceBasePtr(FFoliageInstanceBaseId BaseId) const;
+	
+	/* Returns base info for registered Id */
+	FOLIAGE_API FFoliageInstanceBaseInfo GetInstanceBaseInfo(FFoliageInstanceBaseId BaseId) const;
+	
+	/* Updates base info for a specified base component and returns updated info */
+	FOLIAGE_API FFoliageInstanceBaseInfo UpdateInstanceBaseInfoTransform(UActorComponent* InComponent);
+
+	/* Compacts cross-level references, removing dead links */
+	FOLIAGE_API void CompactInstanceBaseCache(AInstancedFoliageActor* IFA);
+	
+	/* Refreshes base component transforms cache */
+	FOLIAGE_API void UpdateInstanceBaseCachedTransforms();
+	
+	friend FArchive& operator<<(FArchive& Ar, FFoliageInstanceBaseCache& InstanceBaseCache);
+		
+	//
+	static FFoliageInstanceBaseId InvalidBaseId;
+
+	// ID generator
+	// serialized
+	FFoliageInstanceBaseId NextBaseId;
+
+	// Map for looking up base info by ID
+	// serialized
+	TMap<FFoliageInstanceBaseId, FFoliageInstanceBaseInfo> InstanceBaseMap;
+	
+	// Map for looking up base ID by pointer to component
+	// transient
+	TMap<FFoliageInstanceBasePtr, FFoliageInstanceBaseId> InstanceBaseInvMap;
+	
+	// Map for detecting removed bases
+	// serialized
+	TMap<TAssetPtr<UWorld>, TArray<FFoliageInstanceBasePtr>> InstanceBaseLevelMap;
+};
+
+#endif// WITH_EDITORONLY_DATA
\ No newline at end of file
diff --git a/Engine/Source/Runtime/Foliage/Public/InstancedFoliage.h b/Engine/Source/Runtime/Foliage/Public/InstancedFoliage.h
--- a/Engine/Source/Runtime/Foliage/Public/InstancedFoliage.h
+++ b/Engine/Source/Runtime/Foliage/Public/InstancedFoliage.h
@@ -3,13 +3,12 @@
 /*=============================================================================
 	InstancedFoliage.h: Instanced foliage type definitions.
   =============================================================================*/
-
 #pragma once
+#include "FoliageInstanceBase.h"
 
 //
 // Forward declarations.
 //
-class UInstancedStaticMeshComponent;
 class UHierarchicalInstancedStaticMeshComponent;
 class AInstancedFoliageActor;
 class UFoliageType;
@@ -48,22 +47,30 @@ struct FFoliageInstancePlacementInfo
 	{}
 };
 
+/**
+ *	Legacy instance
+ */
+struct FFoliageInstance_Deprecated : public FFoliageInstancePlacementInfo
+{
+	UActorComponent* Base;
+	FGuid ProceduralGuid;
+	friend FArchive& operator<<(FArchive& Ar, FFoliageInstance_Deprecated& Instance);
+};
+
 /**
  *	FFoliageInstance - editor info an individual instance
  */
 struct FFoliageInstance : public FFoliageInstancePlacementInfo
 {
-	UActorComponent* Base;
+	// ID of base this instance was painted on
+	FFoliageInstanceBaseId BaseId;
 
-#if WITH_EDITORONLY_DATA
 	FGuid ProceduralGuid;
-#endif
 
 	FFoliageInstance()
-	: Base(NULL)
+	: BaseId(0)
 	{}
-
-
+	
 	friend FArchive& operator<<(FArchive& Ar, FFoliageInstance& Instance);
 
 	FTransform GetInstanceWorldTransform() const
@@ -100,54 +107,24 @@ struct FFoliageInstance : public FFoliageInstancePlacementInfo
 	}
 };
 
-/**
- * FFoliageComponentHashInfo
- * Cached instance list and component location info stored in the ComponentHash.
- * Used for moving quick updates after operations on components with foliage painted on them.
- */
-struct FFoliageComponentHashInfo
+struct FFoliageMeshInfo_Deprecated
 {
-	// tors
-	FFoliageComponentHashInfo()
-		: CachedLocation(0, 0, 0)
-		, CachedRotation(0, 0, 0)
-		, CachedDrawScale(1, 1, 1)
-	{}
+	UHierarchicalInstancedStaticMeshComponent* Component;
 
-	FFoliageComponentHashInfo(UActorComponent* InComponent)
-		: CachedLocation(0, 0, 0)
-		, CachedRotation(0, 0, 0)
-		, CachedDrawScale(1, 1, 1)
-	{
-		UpdateLocationFromActor(InComponent);
-	}
+#if WITH_EDITORONLY_DATA
+	// Allows us to detect if FoliageType was updated while this level wasn't loaded
+	FGuid FoliageTypeUpdateGuid;
 
-	// Cache the location and rotation from the actor
-	void UpdateLocationFromActor(UActorComponent* InComponent)
+	// Editor-only placed instances
+	TArray<FFoliageInstance_Deprecated> Instances;
+#endif
+
+	FFoliageMeshInfo_Deprecated()
+		: Component(nullptr)
 	{
-		if (InComponent)
-		{
-			AActor* Owner = Cast<AActor>(InComponent->GetOuter());
-			if (Owner)
-			{
-				const USceneComponent* RootComponent = Owner->GetRootComponent();
-				if (RootComponent)
-				{
-					CachedLocation = RootComponent->RelativeLocation;
-					CachedRotation = RootComponent->RelativeRotation;
-					CachedDrawScale = RootComponent->RelativeScale3D;
-				}
-			}
-		}
 	}
 
-	// serializer
-	friend FArchive& operator<<(FArchive& Ar, FFoliageComponentHashInfo& ComponentHashInfo);
-
-	FVector CachedLocation;
-	FRotator CachedRotation;
-	FVector CachedDrawScale;
-	TSet<int32> Instances;
+	friend FArchive& operator<<(FArchive& Ar, FFoliageMeshInfo_Deprecated& MeshInfo);
 };
 
 /**
@@ -168,7 +145,7 @@ struct FFoliageMeshInfo
 	TUniquePtr<FFoliageInstanceHash> InstanceHash;
 
 	// Transient, editor-only set of instances per component
-	TMap<UActorComponent*, FFoliageComponentHashInfo> ComponentHash;
+	TMap<FFoliageInstanceBaseId, TSet<int32>> ComponentHash;
 
 	// Transient, editor-only list of selected instances.
 	TSet<int32> SelectedIndices;
@@ -208,6 +185,7 @@ struct FFoliageMeshInfo
 
 #if WITH_EDITOR
 	FOLIAGE_API void AddInstance(AInstancedFoliageActor* InIFA, const UFoliageType* InSettings, const FFoliageInstance& InNewInstance);
+	FOLIAGE_API void AddInstance(AInstancedFoliageActor* InIFA, const UFoliageType* InSettings, const FFoliageInstance& InNewInstance, UActorComponent* InBaseComponent);
 	FOLIAGE_API void RemoveInstances(AInstancedFoliageActor* InIFA, const TArray<int32>& InInstancesToRemove);
 	FOLIAGE_API void PreMoveInstances(AInstancedFoliageActor* InIFA, const TArray<int32>& InInstancesToMove);
 	FOLIAGE_API void PostMoveInstances(AInstancedFoliageActor* InIFA, const TArray<int32>& InInstancesMoved);
@@ -216,7 +194,7 @@ struct FFoliageMeshInfo
 	FOLIAGE_API void GetInstancesInsideSphere(const FSphere& Sphere, TArray<int32>& OutInstances);
 	FOLIAGE_API bool CheckForOverlappingSphere(const FSphere& Sphere);
 	FOLIAGE_API bool CheckForOverlappingInstanceExcluding(int32 TestInstanceIdx, float Radius, TSet<int32>& ExcludeInstances);
-
+	
 	// Destroy existing clusters and reassign all instances to new clusters
 	FOLIAGE_API void ReallocateClusters(AInstancedFoliageActor* InIFA, UFoliageType* InSettings);
 
@@ -227,6 +205,8 @@ struct FFoliageMeshInfo
 	// Get the number of placed instances
 	FOLIAGE_API int32 GetInstanceCount() const;
 
+	FOLIAGE_API void SetInstanceBaseId(FFoliageInstanceBaseId BaseId, int32 InstanceIdx);
+
 	// For debugging. Validate state after editing.
 	void CheckValid();
 #endif
diff --git a/Engine/Source/Runtime/Foliage/Public/InstancedFoliageActor.h b/Engine/Source/Runtime/Foliage/Public/InstancedFoliageActor.h
--- a/Engine/Source/Runtime/Foliage/Public/InstancedFoliageActor.h
+++ b/Engine/Source/Runtime/Foliage/Public/InstancedFoliageActor.h
@@ -6,6 +6,7 @@
 #include "GameFramework/Actor.h"
 #include "Templates/UniqueObj.h"
 #include "Components/InstancedStaticMeshComponent.h"
+#include "FoliageInstanceBase.h"
 
 #include "InstancedFoliageActor.generated.h"
 
@@ -29,20 +30,26 @@ class AInstancedFoliageActor : public AActor
 	UFoliageType* SelectedMesh;
 
 public:
+#if WITH_EDITORONLY_DATA
+	// Cross level references cache for instances base
+	FFoliageInstanceBaseCache InstanceBaseCache;
+#endif// WITH_EDITORONLY_DATA
+
 	TMap<UFoliageType*, TUniqueObj<FFoliageMeshInfo>> FoliageMeshes;
 
 public:
 	// Begin UObject interface.
 	virtual void Serialize(FArchive& Ar) override;
+	virtual void PostInitProperties() override;
+	virtual void BeginDestroy() override;
 	virtual void PostLoad() override;
+	virtual void PreSave() override;
 	static void AddReferencedObjects(UObject* InThis, FReferenceCollector& Collector);
 	// End UObject interface. 
 
 	// Begin AActor interface.
 	// we don't want to have our components automatically destroyed by the Blueprint code
 	virtual void RerunConstructionScripts() override {}
-	virtual void PostRegisterAllComponents() override;
-	virtual void PostUnregisterAllComponents() override;
 	virtual void ApplyWorldOffset(const FVector& InOffset, bool bWorldShift) override;
 	// End AActor interface.
 
@@ -76,7 +83,7 @@ public:
 	*/
 	static FOLIAGE_API AInstancedFoliageActor* GetInstancedFoliageActorForLevel(ULevel* Level, bool bCreateIfNone = true);
 
-	static FOLIAGE_API AInstancedFoliageActor* FoliageTrace(UWorld* InWorld, FHitResult& OutHit, const FDesiredFoliageInstance& DesiredInstance, const AInstancedFoliageActor* IgnoreIFA = nullptr, FName InTraceTag = NAME_None, bool InbReturnFaceIndex = false);
+	static FOLIAGE_API bool FoliageTrace(UWorld* InWorld, FHitResult& OutHit, const FDesiredFoliageInstance& DesiredInstance, const AInstancedFoliageActor* IgnoreIFA = nullptr, FName InTraceTag = NAME_None, bool InbReturnFaceIndex = false);
 	FOLIAGE_API bool CheckCollisionWithWorld(const UFoliageType* Settings, const FFoliageInstance& Inst, const FVector& HitNormal, const FVector& HitLocation);
 
 #if WITH_EDITOR
@@ -126,9 +133,6 @@ public:
 	// Returns the location for the widget
 	FOLIAGE_API FVector GetSelectionLocation();
 
-	// Transforms Editor specific data which is stored in world space
-	FOLIAGE_API void ApplyLevelTransform(const FTransform& LevelTransform);
-
 	/* Called to notify InstancedFoliageActor that a UFoliageType has been modified */
 	void NotifyFoliageTypeChanged(UFoliageType* FoliageType);
 
@@ -137,11 +141,12 @@ public:
 private:
 #if WITH_EDITOR
 	void OnLevelActorMoved(AActor* InActor);
+	void OnLevelActorDeleted(AActor* InActor);
 #endif
 private:
 #if WITH_EDITOR
 	FDelegateHandle OnLevelActorMovedDelegateHandle;
-	FDelegateHandle OnApplyLevelTransformDelegateHandle;
+	FDelegateHandle OnLevelActorDeletedDelegateHandle;
 #endif
 
 };
diff --git a/Engine/Source/Runtime/Landscape/Private/LandscapeCollision.cpp b/Engine/Source/Runtime/Landscape/Private/LandscapeCollision.cpp
--- a/Engine/Source/Runtime/Landscape/Private/LandscapeCollision.cpp
+++ b/Engine/Source/Runtime/Landscape/Private/LandscapeCollision.cpp
@@ -1019,16 +1019,17 @@ void ULandscapeHeightfieldCollisionComponent::SnapFoliageInstances(AInstancedFol
 		// Find the per-mesh info matching the mesh.
 		UFoliageType* Settings = MeshPair.Key;
 		FFoliageMeshInfo& MeshInfo = *MeshPair.Value;
+		const auto BaseId = IFA.InstanceBaseCache.GetInstanceBaseId(this);  // can we move this out?
 
-		FFoliageComponentHashInfo* ComponentHashInfo = MeshInfo.ComponentHash.Find(this);
-		if (ComponentHashInfo)
+		const auto* InstanceSet = MeshInfo.ComponentHash.Find(BaseId);
+		if (InstanceSet)
 		{
 			float TraceExtentSize = Bounds.SphereRadius * 2.f + 10.f; // extend a little
 			FVector TraceVector = GetOwner()->GetRootComponent()->ComponentToWorld.GetUnitAxis(EAxis::Z) * TraceExtentSize;
 
 			bool bFirst = true;
 			TArray<int32> InstancesToRemove;
-			for (int32 InstanceIndex : ComponentHashInfo->Instances)
+			for (int32 InstanceIndex : *InstanceSet)
 			{
 				FFoliageInstance& Instance = MeshInfo.Instances[InstanceIndex];
 
