Fixed occasional crash related to folliage on deleting actor in the scene

diff --git a/Engine/Source/Editor/FoliageEdit/Private/FoliageEdMode.cpp b/Engine/Source/Editor/FoliageEdit/Private/FoliageEdMode.cpp
--- a/Engine/Source/Editor/FoliageEdit/Private/FoliageEdMode.cpp
+++ b/Engine/Source/Editor/FoliageEdit/Private/FoliageEdMode.cpp
@@ -1505,11 +1505,13 @@ bool FEdModeFoliage::SnapInstanceToGround(AInstancedFoliageActor* InIFA, float A
 			}
 		}
 
+		// Set new base
 		auto NewBaseId = InIFA->InstanceBaseCache.AddInstanceBaseId(HitComponent);
-		Mesh.SetInstanceBaseId(NewBaseId, InstanceIdx);
-	
+		Mesh.RemoveFromHash(InstanceIdx);
+		Instance.BaseId = NewBaseId;
 		Instance.Location = Hit.Location;
 		Instance.ZOffset = 0.f;
+		Mesh.AddToHash(InstanceIdx);
 											
 		if (Instance.Flags & FOLIAGE_AlignToNormal)
 		{
diff --git a/Engine/Source/Runtime/Foliage/Private/FoliageInstanceBase.cpp b/Engine/Source/Runtime/Foliage/Private/FoliageInstanceBase.cpp
--- a/Engine/Source/Runtime/Foliage/Private/FoliageInstanceBase.cpp
+++ b/Engine/Source/Runtime/Foliage/Private/FoliageInstanceBase.cpp
@@ -165,6 +165,8 @@ void FFoliageInstanceBaseCache::CompactInstanceBaseCache(AInstancedFoliageActor*
 		return;
 	}
 
+	FFoliageInstanceBaseCache& Cache = IFA->InstanceBaseCache;
+	
 	TSet<FFoliageInstanceBaseId> BasesInUse;
 	for (auto& Pair : IFA->FoliageMeshes)
 	{
@@ -179,7 +181,7 @@ void FFoliageInstanceBaseCache::CompactInstanceBaseCache(AInstancedFoliageActor*
 	
 	// Look for any removed maps
 	TSet<FFoliageInstanceBasePtr> InvalidBasePtrs;
-	for (auto& Pair : InstanceBaseLevelMap)
+	for (auto& Pair : Cache.InstanceBaseLevelMap)
 	{
 		const auto& WorldAsset = Pair.Key;
 		
@@ -201,7 +203,7 @@ void FFoliageInstanceBaseCache::CompactInstanceBaseCache(AInstancedFoliageActor*
 		if (!bExists)
 		{
 			InvalidBasePtrs.Append(Pair.Value);
-			InstanceBaseLevelMap.Remove(Pair.Key);
+			Cache.InstanceBaseLevelMap.Remove(Pair.Key);
 		}
 		else
 		{
@@ -209,7 +211,7 @@ void FFoliageInstanceBaseCache::CompactInstanceBaseCache(AInstancedFoliageActor*
 			for (int32 i = Pair.Value.Num()-1; i >= 0; --i)
 			{
 				// Base needs to be removed if it's not in use by existing instances or component was removed
-				if (Pair.Value[i].IsNull() || !BasesInUse.Contains(GetInstanceBaseId(Pair.Value[i])))
+				if (Pair.Value[i].IsNull() || !BasesInUse.Contains(Cache.GetInstanceBaseId(Pair.Value[i])))
 				{
 					InvalidBasePtrs.Add(Pair.Value[i]);
 					Pair.Value.RemoveAt(i);
@@ -218,27 +220,27 @@ void FFoliageInstanceBaseCache::CompactInstanceBaseCache(AInstancedFoliageActor*
 
 			if (Pair.Value.Num() == 0)
 			{
-				InstanceBaseLevelMap.Remove(Pair.Key);
+				Cache.InstanceBaseLevelMap.Remove(Pair.Key);
 			}
 		}
 	}
 	
 	TSet<FFoliageInstanceBaseId> InvalidBaseIds;
-	InstanceBaseInvMap.Empty();
+	Cache.InstanceBaseInvMap.Empty();
 	// Look for any removed base components
-	for (const auto& Pair : InstanceBaseMap)
+	for (const auto& Pair : Cache.InstanceBaseMap)
 	{
 		const FFoliageInstanceBaseInfo& BaseInfo = Pair.Value;
 		if (InvalidBasePtrs.Contains(BaseInfo.BasePtr))
 		{
 			InvalidBaseIds.Add(Pair.Key);
-			InstanceBaseMap.Remove(Pair.Key);
+			Cache.InstanceBaseMap.Remove(Pair.Key);
 		}
 		else
 		{
 			// Regenerate inverse map
-			check(!InstanceBaseInvMap.Contains(BaseInfo.BasePtr));
-			InstanceBaseInvMap.Add(BaseInfo.BasePtr, Pair.Key);
+			check(!Cache.InstanceBaseInvMap.Contains(BaseInfo.BasePtr));
+			Cache.InstanceBaseInvMap.Add(BaseInfo.BasePtr, Pair.Key);
 		}
 	}
 
@@ -247,24 +249,23 @@ void FFoliageInstanceBaseCache::CompactInstanceBaseCache(AInstancedFoliageActor*
 		for (auto& Pair : IFA->FoliageMeshes)
 		{
 			auto& MeshInfo = Pair.Value;
+			MeshInfo->ComponentHash.Empty();
+			int32 InstanceIdx = 0;
+			
 			for (FFoliageInstance& Instance : MeshInfo->Instances)
 			{
 				if (InvalidBaseIds.Contains(Instance.BaseId))
 				{
 					Instance.BaseId = FFoliageInstanceBaseCache::InvalidBaseId;
 				}
-			}
 
-			for (auto RemovedBaseId : InvalidBaseIds)
-			{
-				MeshInfo->ComponentHash.Remove(RemovedBaseId);
+				MeshInfo->ComponentHash.FindOrAdd(Instance.BaseId).Add(InstanceIdx);
+				InstanceIdx++;
 			}
-		
-			MeshInfo->ComponentHash.Compact();
 		}
 
-		InstanceBaseMap.Compact();
-		InstanceBaseLevelMap.Compact();
+		Cache.InstanceBaseMap.Compact();
+		Cache.InstanceBaseLevelMap.Compact();
 	}
 }
 
diff --git a/Engine/Source/Runtime/Foliage/Private/InstancedFoliage.cpp b/Engine/Source/Runtime/Foliage/Private/InstancedFoliage.cpp
--- a/Engine/Source/Runtime/Foliage/Private/InstancedFoliage.cpp
+++ b/Engine/Source/Runtime/Foliage/Private/InstancedFoliage.cpp
@@ -525,8 +525,7 @@ void FFoliageMeshInfo::AddInstance(AInstancedFoliageActor* InIFA, const UFoliage
 	FFoliageInstance& AddedInstance = Instances[InstanceIndex];
 
 	// Add the instance to the hash
-	InstanceHash->InsertInstance(InNewInstance.Location, InstanceIndex);
-	ComponentHash.FindOrAdd(InNewInstance.BaseId).Add(InstanceIndex);
+	AddToHash(InstanceIndex);
 	// Calculate transform for the instance
 	FTransform InstanceToWorld = InNewInstance.GetInstanceWorldTransform();
 
@@ -565,8 +564,7 @@ void FFoliageMeshInfo::RemoveInstances(AInstancedFoliageActor* InIFA, const TArr
 			FFoliageInstance& Instance = Instances[InstanceIndex];
 
 			// remove from hash
-			InstanceHash->RemoveInstance(Instance.Location, InstanceIndex);
-			SetInstanceBaseId(FFoliageInstanceBaseCache::InvalidBaseId, InstanceIndex);
+			RemoveFromHash(InstanceIndex);
 
 			// remove from the component
 			Component->RemoveInstance(InstanceIndex);
@@ -678,34 +676,32 @@ int32 FFoliageMeshInfo::GetInstanceCount() const
 	return Instances.Num();
 }
 
-void FFoliageMeshInfo::SetInstanceBaseId(FFoliageInstanceBaseId NewBaseId, int32 InstanceIndex)
+void FFoliageMeshInfo::AddToHash(int32 InstanceIndex)
 {
 	FFoliageInstance& Instance = Instances[InstanceIndex];
-	if (NewBaseId != Instance.BaseId)
-	{
-		// Remove current base link
-		auto* InstanceSet = ComponentHash.Find(Instance.BaseId);
-		if (InstanceSet)
-		{
-			InstanceSet->Remove(InstanceIndex);
-			if (InstanceSet->Num() == 0)
-			{
-				// Remove the component from the component hash if this is the last instance.
-				ComponentHash.Remove(Instance.BaseId);
-			}
-		}
+	InstanceHash->InsertInstance(Instance.Location, InstanceIndex);
+	ComponentHash.FindOrAdd(Instance.BaseId).Add(InstanceIndex);
+}
 
-		// Assign new
-		Instance.BaseId = NewBaseId;
-		
-		if (NewBaseId != FFoliageInstanceBaseCache::InvalidBaseId)
+void FFoliageMeshInfo::RemoveFromHash(int32 InstanceIndex)
+{
+	FFoliageInstance& Instance = Instances[InstanceIndex];
+	
+	InstanceHash->RemoveInstance(Instance.Location, InstanceIndex);
+
+	// Remove current base link
+	auto* InstanceSet = ComponentHash.Find(Instance.BaseId);
+	if (InstanceSet)
+	{
+		InstanceSet->Remove(InstanceIndex);
+		if (InstanceSet->Num() == 0)
 		{
-			ComponentHash.FindOrAdd(NewBaseId).Add(InstanceIndex);
+			// Remove the component from the component hash if this is the last instance.
+			ComponentHash.Remove(Instance.BaseId);
 		}
 	}
 }
 
-
 // Destroy existing clusters and reassign all instances to new clusters
 void FFoliageMeshInfo::ReallocateClusters(AInstancedFoliageActor* InIFA, UFoliageType* InSettings)
 {
@@ -1752,12 +1748,7 @@ void AInstancedFoliageActor::PostLoad()
 			MeshInfo.InstanceHash->Empty();
 			for (int32 InstanceIdx = 0; InstanceIdx < MeshInfo.Instances.Num(); InstanceIdx++)
 			{
-				FFoliageInstance& Instance = MeshInfo.Instances[InstanceIdx];
-				{
-					// Add valid instances to the hash.
-					MeshInfo.InstanceHash->InsertInstance(Instance.Location, InstanceIdx);
-					MeshInfo.ComponentHash.FindOrAdd(Instance.BaseId).Add(InstanceIdx);
-				}
+				MeshInfo.AddToHash(InstanceIdx);
 			}
 			
 			// Fix-up cross level references after serializing deprecated data
@@ -1777,7 +1768,7 @@ void AInstancedFoliageActor::PostLoad()
 			}
 
 			// Clean up dead cross-level references
-			InstanceBaseCache.CompactInstanceBaseCache(this);
+			FFoliageInstanceBaseCache::CompactInstanceBaseCache(this);
 			
 			// Convert to Heirarchical foliage
 			if (GetLinkerCustomVersion(FFoliageCustomVersion::GUID) < FFoliageCustomVersion::FoliageUsingHierarchicalISMC)
@@ -1804,7 +1795,7 @@ void AInstancedFoliageActor::PreSave()
 #if WITH_EDITOR
 	if (GIsEditor)
 	{
-		InstanceBaseCache.CompactInstanceBaseCache(this);
+		FFoliageInstanceBaseCache::CompactInstanceBaseCache(this);
 	}
 #endif //WITH_EDITOR
 }
diff --git a/Engine/Source/Runtime/Foliage/Public/FoliageInstanceBase.h b/Engine/Source/Runtime/Foliage/Public/FoliageInstanceBase.h
--- a/Engine/Source/Runtime/Foliage/Public/FoliageInstanceBase.h
+++ b/Engine/Source/Runtime/Foliage/Public/FoliageInstanceBase.h
@@ -53,7 +53,7 @@ struct FFoliageInstanceBaseCache
 	FOLIAGE_API FFoliageInstanceBaseInfo UpdateInstanceBaseInfoTransform(UActorComponent* InComponent);
 
 	/* Compacts cross-level references, removing dead links */
-	FOLIAGE_API void CompactInstanceBaseCache(AInstancedFoliageActor* IFA);
+	FOLIAGE_API static void CompactInstanceBaseCache(AInstancedFoliageActor* IFA);
 	
 	/* Refreshes base component transforms cache */
 	FOLIAGE_API void UpdateInstanceBaseCachedTransforms();
diff --git a/Engine/Source/Runtime/Foliage/Public/InstancedFoliage.h b/Engine/Source/Runtime/Foliage/Public/InstancedFoliage.h
--- a/Engine/Source/Runtime/Foliage/Public/InstancedFoliage.h
+++ b/Engine/Source/Runtime/Foliage/Public/InstancedFoliage.h
@@ -205,8 +205,9 @@ struct FFoliageMeshInfo
 	// Get the number of placed instances
 	FOLIAGE_API int32 GetInstanceCount() const;
 
-	FOLIAGE_API void SetInstanceBaseId(FFoliageInstanceBaseId BaseId, int32 InstanceIdx);
-
+	FOLIAGE_API void AddToHash(int32 InstanceIdx);
+	FOLIAGE_API void RemoveFromHash(int32 InstanceIdx);
+		
 	// For debugging. Validate state after editing.
 	void CheckValid();
 #endif
