Temp fix for old foliage being deleted whenever delete actor delegate is called.

diff --git a/Engine/Source/Runtime/Foliage/Private/InstancedFoliage.cpp b/Engine/Source/Runtime/Foliage/Private/InstancedFoliage.cpp
--- a/Engine/Source/Runtime/Foliage/Private/InstancedFoliage.cpp
+++ b/Engine/Source/Runtime/Foliage/Private/InstancedFoliage.cpp
@@ -1082,6 +1082,13 @@ void AInstancedFoliageActor::DeleteInstancesForComponent(UActorComponent* InComp
 	{
 		FFoliageMeshInfo& MeshInfo = *MeshPair.Value;
 		const auto BaseId = InstanceBaseCache.GetInstanceBaseId(InComponent);
+
+		//TEMP FIX. Not sure why we're getting instances with invalid index, but we should not delete these
+		if (BaseId == FFoliageInstanceBaseCache::InvalidBaseId)
+		{
+			continue;
+		}
+
 		const auto* InstanceSet = MeshInfo.ComponentHash.Find(BaseId);
 		if (InstanceSet)
 		{
