Foliage: fixed missing bases after map is saved

diff --git a/Engine/Source/Runtime/Foliage/Private/FoliageInstanceBase.cpp b/Engine/Source/Runtime/Foliage/Private/FoliageInstanceBase.cpp
--- a/Engine/Source/Runtime/Foliage/Private/FoliageInstanceBase.cpp
+++ b/Engine/Source/Runtime/Foliage/Private/FoliageInstanceBase.cpp
@@ -93,10 +93,10 @@ FFoliageInstanceBaseId FFoliageInstanceBaseCache::AddInstanceBaseId(UActorCompon
 		if (BaseId == FFoliageInstanceBaseCache::InvalidBaseId)
 		{
 			BaseId = NextBaseId++;
-			FFoliageInstanceBasePtr BasePtr = InComponent;
 
-			InstanceBaseMap.Add(BaseId, FFoliageInstanceBaseInfo(InComponent));
-			InstanceBaseInvMap.Add(BasePtr, BaseId);
+			FFoliageInstanceBaseInfo BaseInfo(InComponent);
+			InstanceBaseMap.Add(BaseId, BaseInfo);
+			InstanceBaseInvMap.Add(BaseInfo.BasePtr, BaseId);
 
 			check(InstanceBaseMap.Num() == InstanceBaseInvMap.Num());
 
@@ -107,7 +107,7 @@ FFoliageInstanceBaseId FFoliageInstanceBaseCache::AddInstanceBaseId(UActorCompon
 				if (ComponentWorld)
 				{
 					auto WorldKey = TAssetPtr<UWorld>(ComponentWorld);
-					InstanceBaseLevelMap.FindOrAdd(WorldKey).Add(BasePtr);
+					InstanceBaseLevelMap.FindOrAdd(WorldKey).Add(BaseInfo.BasePtr);
 				}
 			}
 		}
diff --git a/Engine/Source/Runtime/Foliage/Private/InstancedFoliage.cpp b/Engine/Source/Runtime/Foliage/Private/InstancedFoliage.cpp
--- a/Engine/Source/Runtime/Foliage/Private/InstancedFoliage.cpp
+++ b/Engine/Source/Runtime/Foliage/Private/InstancedFoliage.cpp
@@ -1799,10 +1799,7 @@ void AInstancedFoliageActor::PostLoad()
 					BaseList.Add(BaseInfoPair.Value.BasePtr);
 				}
 			}
-
-			// Clean up dead cross-level references
-			FFoliageInstanceBaseCache::CompactInstanceBaseCache(this);
-			
+		
 			// Convert to Heirarchical foliage
 			if (GetLinkerCustomVersion(FFoliageCustomVersion::GUID) < FFoliageCustomVersion::FoliageUsingHierarchicalISMC)
 			{
@@ -1817,20 +1814,12 @@ void AInstancedFoliageActor::PostLoad()
 				}
 			}
 		}
-	}
-#endif
-}
 
-void AInstancedFoliageActor::PreSave()
-{
-	Super::PreSave();
-
-#if WITH_EDITOR
-	if (GIsEditor)
-	{
+		// Clean up dead cross-level references
 		FFoliageInstanceBaseCache::CompactInstanceBaseCache(this);
+
 	}
-#endif //WITH_EDITOR
+#endif
 }
 
 #if WITH_EDITOR
