Fixed legacy foliage data serialization
Fixed crash on foliage snap (End key)

diff --git a/Engine/Source/Editor/FoliageEdit/Private/FoliageEdMode.cpp b/Engine/Source/Editor/FoliageEdit/Private/FoliageEdMode.cpp
--- a/Engine/Source/Editor/FoliageEdit/Private/FoliageEdMode.cpp
+++ b/Engine/Source/Editor/FoliageEdit/Private/FoliageEdMode.cpp
@@ -1516,12 +1516,12 @@ bool FEdModeFoliage::SnapInstanceToGround(AInstancedFoliageActor* InIFA, float A
 
 		// Set new base
 		auto NewBaseId = InIFA->InstanceBaseCache.AddInstanceBaseId(HitComponent);
-		Mesh.RemoveFromHash(InstanceIdx);
+		Mesh.RemoveFromBaseHash(InstanceIdx);
 		Instance.BaseId = NewBaseId;
+		Mesh.AddToBaseHash(InstanceIdx);
 		Instance.Location = Hit.Location;
 		Instance.ZOffset = 0.f;
-		Mesh.AddToHash(InstanceIdx);
-											
+													
 		if (Instance.Flags & FOLIAGE_AlignToNormal)
 		{
 			// Remove previous alignment and align to new normal.
diff --git a/Engine/Source/Runtime/Foliage/Private/FoliageInstanceBase.cpp b/Engine/Source/Runtime/Foliage/Private/FoliageInstanceBase.cpp
--- a/Engine/Source/Runtime/Foliage/Private/FoliageInstanceBase.cpp
+++ b/Engine/Source/Runtime/Foliage/Private/FoliageInstanceBase.cpp
@@ -69,7 +69,7 @@ FArchive& operator << (FArchive& Ar, FFoliageInstanceBaseCache& InstanceBaseCach
 	{
 		Ar << InstanceBaseCache.InstanceBaseInvMap;
 	}
-	else if (Ar.IsLoading())
+	else if (!Ar.IsSaving())
 	{
 		InstanceBaseCache.InstanceBaseInvMap.Empty();
 		for (const auto& Pair : InstanceBaseCache.InstanceBaseMap)
@@ -78,8 +78,6 @@ FArchive& operator << (FArchive& Ar, FFoliageInstanceBaseCache& InstanceBaseCach
 			InstanceBaseCache.InstanceBaseInvMap.Add(BaseInfo.BasePtr, Pair.Key);
 		}
 	}
-	
-	check(InstanceBaseCache.InstanceBaseInvMap.Num() == InstanceBaseCache.InstanceBaseMap.Num());
 
 	return Ar;
 }
@@ -160,7 +158,7 @@ FFoliageInstanceBaseInfo FFoliageInstanceBaseCache::UpdateInstanceBaseInfoTransf
 void FFoliageInstanceBaseCache::CompactInstanceBaseCache(AInstancedFoliageActor* IFA)
 {
 	UWorld* World = IFA->GetWorld();
-	if (!World)
+	if (!World || World->IsGameWorld())
 	{
 		return;
 	}
diff --git a/Engine/Source/Runtime/Foliage/Private/InstancedFoliage.cpp b/Engine/Source/Runtime/Foliage/Private/InstancedFoliage.cpp
--- a/Engine/Source/Runtime/Foliage/Private/InstancedFoliage.cpp
+++ b/Engine/Source/Runtime/Foliage/Private/InstancedFoliage.cpp
@@ -141,6 +141,16 @@ static void ConvertDeprecatedFoliageMeshes(
 			FoliageMesh->Instances.Add(Instance);
 		}
 	}
+
+	// there were no cross-level references before
+	check(IFA->InstanceBaseCache.InstanceBaseLevelMap.Num() <= 1); 
+	// populate WorldAsset->BasePtr map
+	IFA->InstanceBaseCache.InstanceBaseLevelMap.Empty();
+	auto& BaseList = IFA->InstanceBaseCache.InstanceBaseLevelMap.Add(TAssetPtr<UWorld>(Cast<UWorld>(IFA->GetLevel()->GetOuter())));
+	for (auto& BaseInfoPair : IFA->InstanceBaseCache.InstanceBaseMap)
+	{
+		BaseList.Add(BaseInfoPair.Value.BasePtr);
+	}
 #endif//WITH_EDITORONLY_DATA	
 }
 
@@ -525,7 +535,8 @@ void FFoliageMeshInfo::AddInstance(AInstancedFoliageActor* InIFA, const UFoliage
 	FFoliageInstance& AddedInstance = Instances[InstanceIndex];
 
 	// Add the instance to the hash
-	AddToHash(InstanceIndex);
+	AddToBaseHash(InstanceIndex);
+	InstanceHash->InsertInstance(AddedInstance.Location, InstanceIndex);
 	// Calculate transform for the instance
 	FTransform InstanceToWorld = InNewInstance.GetInstanceWorldTransform();
 
@@ -564,7 +575,8 @@ void FFoliageMeshInfo::RemoveInstances(AInstancedFoliageActor* InIFA, const TArr
 			FFoliageInstance& Instance = Instances[InstanceIndex];
 
 			// remove from hash
-			RemoveFromHash(InstanceIndex);
+			RemoveFromBaseHash(InstanceIndex);
+			InstanceHash->RemoveInstance(Instance.Location, InstanceIndex);
 
 			// remove from the component
 			Component->RemoveInstance(InstanceIndex);
@@ -676,18 +688,15 @@ int32 FFoliageMeshInfo::GetInstanceCount() const
 	return Instances.Num();
 }
 
-void FFoliageMeshInfo::AddToHash(int32 InstanceIndex)
+void FFoliageMeshInfo::AddToBaseHash(int32 InstanceIndex)
 {
 	FFoliageInstance& Instance = Instances[InstanceIndex];
-	InstanceHash->InsertInstance(Instance.Location, InstanceIndex);
 	ComponentHash.FindOrAdd(Instance.BaseId).Add(InstanceIndex);
 }
 
-void FFoliageMeshInfo::RemoveFromHash(int32 InstanceIndex)
+void FFoliageMeshInfo::RemoveFromBaseHash(int32 InstanceIndex)
 {
 	FFoliageInstance& Instance = Instances[InstanceIndex];
-	
-	InstanceHash->RemoveInstance(Instance.Location, InstanceIndex);
 
 	// Remove current base link
 	auto* InstanceSet = ComponentHash.Find(Instance.BaseId);
@@ -1658,18 +1667,14 @@ void AInstancedFoliageActor::Serialize(FArchive& Ar)
 #if WITH_EDITORONLY_DATA
 			NewMeshInfo.FoliageTypeUpdateGuid = FoliageType->UpdateGuid;
 #endif
-			FoliageMeshesDeprecated.Add(FoliageType, TUniqueObj<FFoliageMeshInfo_Deprecated>(MoveTemp(NewMeshInfo)));
+			FoliageMeshes_Deprecated.Add(FoliageType, TUniqueObj<FFoliageMeshInfo_Deprecated>(MoveTemp(NewMeshInfo)));
 		}
-
-		ConvertDeprecatedFoliageMeshes(this, FoliageMeshesDeprecated, FoliageMeshes);
 	}
 	else
 	{
 		if (Ar.CustomVer(FFoliageCustomVersion::GUID) < FFoliageCustomVersion::CrossLevelBase)
 		{
-			TMap<UFoliageType*, TUniqueObj<FFoliageMeshInfo_Deprecated>> FoliageMeshesDeprecated;
-			Ar << FoliageMeshesDeprecated;
-			ConvertDeprecatedFoliageMeshes(this, FoliageMeshesDeprecated, FoliageMeshes);
+			Ar << FoliageMeshes_Deprecated;
 		}
 		else
 		{
@@ -1725,6 +1730,12 @@ void AInstancedFoliageActor::PostLoad()
 #if WITH_EDITOR
 	if (GIsEditor)
 	{
+		if (GetLinkerCustomVersion(FFoliageCustomVersion::GUID) < FFoliageCustomVersion::CrossLevelBase)
+		{
+			ConvertDeprecatedFoliageMeshes(this, FoliageMeshes_Deprecated, FoliageMeshes);
+			FoliageMeshes_Deprecated.Empty();
+		}
+				
 		{
 			bool bContainsNull = FoliageMeshes.Remove(nullptr) > 0;
 			if (bContainsNull)
@@ -1781,25 +1792,10 @@ void AInstancedFoliageActor::PostLoad()
 			MeshInfo.InstanceHash->Empty();
 			for (int32 InstanceIdx = 0; InstanceIdx < MeshInfo.Instances.Num(); InstanceIdx++)
 			{
-				MeshInfo.AddToHash(InstanceIdx);
+				MeshInfo.AddToBaseHash(InstanceIdx);
+				MeshInfo.InstanceHash->InsertInstance(MeshInfo.Instances[InstanceIdx].Location, InstanceIdx);
 			}
-			
-			// Fix-up cross level references after serializing deprecated data
-			if (GetLinkerCustomVersion(FFoliageCustomVersion::GUID) < FFoliageCustomVersion::CrossLevelBase)
-			{
-				// update base cached transform this information was not available during serialization
-				InstanceBaseCache.UpdateInstanceBaseCachedTransforms();
-				// there were no cross-level references before
-				check(InstanceBaseCache.InstanceBaseLevelMap.Num() <= 1); 
-				// populate WorldAsset->BasePtr map
-				InstanceBaseCache.InstanceBaseLevelMap.Empty();
-				auto& BaseList = InstanceBaseCache.InstanceBaseLevelMap.Add(TAssetPtr<UWorld>(Cast<UWorld>(GetLevel()->GetOuter())));
-				for (auto& BaseInfoPair : InstanceBaseCache.InstanceBaseMap)
-				{
-					BaseList.Add(BaseInfoPair.Value.BasePtr);
-				}
-			}
-		
+	
 			// Convert to Heirarchical foliage
 			if (GetLinkerCustomVersion(FFoliageCustomVersion::GUID) < FFoliageCustomVersion::FoliageUsingHierarchicalISMC)
 			{
diff --git a/Engine/Source/Runtime/Foliage/Public/InstancedFoliage.h b/Engine/Source/Runtime/Foliage/Public/InstancedFoliage.h
--- a/Engine/Source/Runtime/Foliage/Public/InstancedFoliage.h
+++ b/Engine/Source/Runtime/Foliage/Public/InstancedFoliage.h
@@ -205,8 +205,8 @@ struct FFoliageMeshInfo
 	// Get the number of placed instances
 	FOLIAGE_API int32 GetInstanceCount() const;
 
-	FOLIAGE_API void AddToHash(int32 InstanceIdx);
-	FOLIAGE_API void RemoveFromHash(int32 InstanceIdx);
+	FOLIAGE_API void AddToBaseHash(int32 InstanceIdx);
+	FOLIAGE_API void RemoveFromBaseHash(int32 InstanceIdx);
 		
 	// For debugging. Validate state after editing.
 	void CheckValid();
diff --git a/Engine/Source/Runtime/Foliage/Public/InstancedFoliageActor.h b/Engine/Source/Runtime/Foliage/Public/InstancedFoliageActor.h
--- a/Engine/Source/Runtime/Foliage/Public/InstancedFoliageActor.h
+++ b/Engine/Source/Runtime/Foliage/Public/InstancedFoliageActor.h
@@ -43,8 +43,6 @@ public:
 	virtual void PostInitProperties() override;
 	virtual void BeginDestroy() override;
 	virtual void PostLoad() override;
-	// @todo beng: fixme
-	//virtual void PreSave() override;
 	static void AddReferencedObjects(UObject* InThis, FReferenceCollector& Collector);
 	// End UObject interface. 
 
@@ -144,6 +142,11 @@ public:
 #endif	//WITH_EDITOR
 
 private:
+#if WITH_EDITORONLY_DATA
+	// Deprecated data, will be converted and cleaned up in PostLoad
+	TMap<UFoliageType*, TUniqueObj<struct FFoliageMeshInfo_Deprecated>> FoliageMeshes_Deprecated;
+#endif//WITH_EDITORONLY_DATA
+	
 #if WITH_EDITOR
 	void OnLevelActorMoved(AActor* InActor);
 	void OnLevelActorDeleted(AActor* InActor);
