Fixed non-editor build

diff --git a/Engine/Source/Runtime/Foliage/Private/InstancedFoliage.cpp b/Engine/Source/Runtime/Foliage/Private/InstancedFoliage.cpp
--- a/Engine/Source/Runtime/Foliage/Private/InstancedFoliage.cpp
+++ b/Engine/Source/Runtime/Foliage/Private/InstancedFoliage.cpp
@@ -1635,6 +1635,7 @@ void AInstancedFoliageActor::Serialize(FArchive& Ar)
 	
 	if (Ar.UE4Ver() < VER_UE4_FOLIAGE_SETTINGS_TYPE)
 	{
+#if WITH_EDITORONLY_DATA
 		TMap<UFoliageType*, TUniqueObj<FFoliageMeshInfo_Deprecated>> FoliageMeshesDeprecated;
 		TMap<UStaticMesh*, FFoliageMeshInfo_Old> OldFoliageMeshes;
 		Ar << OldFoliageMeshes;
@@ -1642,10 +1643,8 @@ void AInstancedFoliageActor::Serialize(FArchive& Ar)
 		{
 			FFoliageMeshInfo_Deprecated NewMeshInfo;
 
-#if WITH_EDITORONLY_DATA
 			NewMeshInfo.Instances = MoveTemp(OldMeshInfo.Value.Instances);
-#endif
-
+			
 			UFoliageType_InstancedStaticMesh* FoliageType = OldMeshInfo.Value.Settings;
 			if (FoliageType == nullptr)
 			{
@@ -1664,17 +1663,18 @@ void AInstancedFoliageActor::Serialize(FArchive& Ar)
 				FoliageType = (UFoliageType_InstancedStaticMesh*)StaticDuplicateObject(FoliageType, this, nullptr, RF_AllFlags & ~(RF_Standalone | RF_Public));
 				FoliageType->Mesh = OldMeshInfo.Key;
 			}
-#if WITH_EDITORONLY_DATA
 			NewMeshInfo.FoliageTypeUpdateGuid = FoliageType->UpdateGuid;
-#endif
 			FoliageMeshes_Deprecated.Add(FoliageType, TUniqueObj<FFoliageMeshInfo_Deprecated>(MoveTemp(NewMeshInfo)));
 		}
+#endif//WITH_EDITORONLY_DATA
 	}
 	else
 	{
 		if (Ar.CustomVer(FFoliageCustomVersion::GUID) < FFoliageCustomVersion::CrossLevelBase)
 		{
+#if WITH_EDITORONLY_DATA
 			Ar << FoliageMeshes_Deprecated;
+#endif
 		}
 		else
 		{
