Updated comment in AInstancedFoliageActor::DeleteInstancesForComponent

diff --git a/Engine/Source/Runtime/Foliage/Private/InstancedFoliage.cpp b/Engine/Source/Runtime/Foliage/Private/InstancedFoliage.cpp
--- a/Engine/Source/Runtime/Foliage/Private/InstancedFoliage.cpp
+++ b/Engine/Source/Runtime/Foliage/Private/InstancedFoliage.cpp
@@ -1092,7 +1092,7 @@ void AInstancedFoliageActor::DeleteInstancesForComponent(UActorComponent* InComp
 		FFoliageMeshInfo& MeshInfo = *MeshPair.Value;
 		const auto BaseId = InstanceBaseCache.GetInstanceBaseId(InComponent);
 
-		//TEMP FIX. Not sure why we're getting instances with invalid index, but we should not delete these
+		// Instances with empty base has BaseId==InvalidBaseId, we should not delete these
 		if (BaseId == FFoliageInstanceBaseCache::InvalidBaseId)
 		{
 			continue;
