Turn procedural foliage actor into volumes

diff --git a/Engine/Source/Editor/FoliageEdit/Private/ActorFactoryProceduralFoliage.cpp b/Engine/Source/Editor/FoliageEdit/Private/ActorFactoryProceduralFoliage.cpp
--- a/Engine/Source/Editor/FoliageEdit/Private/ActorFactoryProceduralFoliage.cpp
+++ b/Engine/Source/Editor/FoliageEdit/Private/ActorFactoryProceduralFoliage.cpp
@@ -38,8 +38,8 @@ bool UActorFactoryProceduralFoliage::CanCreateActorFrom(const FAssetData& AssetD
 
 void UActorFactoryProceduralFoliage::PostSpawnActor(UObject* Asset, AActor* NewActor)
 {
+	Super::PostSpawnActor(Asset, NewActor);
 	UProceduralFoliage* ProceduralFoliage = CastChecked<UProceduralFoliage>(Asset);
-	GEditor->SetActorLabelUnique(NewActor, ProceduralFoliage->GetName());
 
 	UE_LOG(LogActorFactory, Log, TEXT("Actor Factory created %s"), *ProceduralFoliage->GetName());
 
diff --git a/Engine/Source/Editor/FoliageEdit/Private/ActorFactoryProceduralFoliage.h b/Engine/Source/Editor/FoliageEdit/Private/ActorFactoryProceduralFoliage.h
--- a/Engine/Source/Editor/FoliageEdit/Private/ActorFactoryProceduralFoliage.h
+++ b/Engine/Source/Editor/FoliageEdit/Private/ActorFactoryProceduralFoliage.h
@@ -2,11 +2,11 @@
 
 
 #pragma once
-#include "ActorFactories/ActorFactory.h"
+#include "ActorFactories/ActorFactoryBoxVolume.h"
 #include "ActorFactoryProceduralFoliage.generated.h"
 
 UCLASS(MinimalAPI, config=Editor)
-class UActorFactoryProceduralFoliage : public UActorFactory
+class UActorFactoryProceduralFoliage : public UActorFactoryBoxVolume
 {
 	GENERATED_UCLASS_BODY()
 
diff --git a/Engine/Source/Editor/FoliageEdit/Private/ProceduralFoliageComponentVisualizer.cpp b/Engine/Source/Editor/FoliageEdit/Private/ProceduralFoliageComponentVisualizer.cpp
--- a/Engine/Source/Editor/FoliageEdit/Private/ProceduralFoliageComponentVisualizer.cpp
+++ b/Engine/Source/Editor/FoliageEdit/Private/ProceduralFoliageComponentVisualizer.cpp
@@ -13,16 +13,23 @@ void FProceduralFoliageComponentVisualizer::DrawVisualization(const UActorCompon
 {
 	if (const UProceduralFoliageComponent* ProcComponent = Cast<const UProceduralFoliageComponent>(Component))
 	{
-		if (ProcComponent->ProceduralFoliage)
+		if (ProcComponent->ProceduralFoliage && ProcComponent->bHideDebugTiles == false)
 		{
 			const float TileSize = ProcComponent->ProceduralFoliage->TileSize;
 			const FVector TileSizeOffset(TileSize, TileSize, 0.f);
-			const FVector TileLocation = ProcComponent->GetComponentToWorld().GetLocation();
+			const FVector TileLocation = ProcComponent->GetWorldPosition();
 			const FVector OverlapV(ProcComponent->Overlap, ProcComponent->Overlap, 0.f);
 			
-			for (int32 X = 0; X < ProcComponent->TilesX; ++X)
+			int32 XOffset, YOffset;
+			int32 NumX = 0;
+			int32 NumY = 0;
+			float HalfHeight;
+
+			ProcComponent->GetTilesLayout(XOffset, YOffset, NumX, NumY, HalfHeight);
+
+			for (int32 X = 0; X < NumX; ++X)
 			{
-				for (int32 Y = 0; Y < ProcComponent->TilesY; ++Y)
+				for (int32 Y = 0; Y < NumY; ++Y)
 				{
 					const FVector StartOffset(X*TileSize, Y*TileSize, 0.f);
 					
diff --git a/Engine/Source/Editor/UnrealEd/Classes/ActorFactories/ActorFactoryBoxVolume.h b/Engine/Source/Editor/UnrealEd/Classes/ActorFactories/ActorFactoryBoxVolume.h
--- a/Engine/Source/Editor/UnrealEd/Classes/ActorFactories/ActorFactoryBoxVolume.h
+++ b/Engine/Source/Editor/UnrealEd/Classes/ActorFactories/ActorFactoryBoxVolume.h
@@ -12,6 +12,6 @@ class UActorFactoryBoxVolume : public UActorFactory
 
 	// Begin UActorFactory Interface
 	virtual bool CanCreateActorFrom( const FAssetData& AssetData, FText& OutErrorMsg ) override;
-	virtual void PostSpawnActor( UObject* Asset, AActor* NewActor ) override;
+	UNREALED_API virtual void PostSpawnActor( UObject* Asset, AActor* NewActor ) override;
 	// End UActorFactory Interface
 };
diff --git a/Engine/Source/Runtime/Foliage/Private/InstancedFoliage.cpp b/Engine/Source/Runtime/Foliage/Private/InstancedFoliage.cpp
--- a/Engine/Source/Runtime/Foliage/Private/InstancedFoliage.cpp
+++ b/Engine/Source/Runtime/Foliage/Private/InstancedFoliage.cpp
@@ -1956,10 +1956,9 @@ bool AInstancedFoliageActor::FoliageTrace(const UWorld* InWorld, FHitResult& Out
 	FVector StartTrace = DesiredInstance.StartTrace;
 
 	TArray<FHitResult> Hits;
-	FCollisionShape SphereShape;
-	SphereShape.SetSphere(DesiredInstance.TraceRadius);
 
-	InWorld->SweepMulti(Hits, StartTrace, DesiredInstance.EndTrace, FQuat::Identity, SphereShape, QueryParams, FCollisionObjectQueryParams(ECC_WorldStatic));
+	bool bInsideProceduralVolume = false;
+	InWorld->LineTraceMulti(Hits, StartTrace, DesiredInstance.EndTrace, QueryParams, FCollisionObjectQueryParams(ECC_WorldStatic));
 	for (const FHitResult& Hit : Hits)
 	{
 		if (DesiredInstance.PlacementMode == EFoliagePlacementMode::Procedural)
@@ -1978,6 +1977,15 @@ bool AInstancedFoliageActor::FoliageTrace(const UWorld* InWorld, FHitResult& Out
 				{
 					return false;
 				}
+				else if (Cast<AProceduralFoliageActor>(Hit.Actor.Get()))	//we never want to collide with our spawning volume
+				{
+					continue;
+				}
+
+				if (bInsideProceduralVolume == false)
+				{
+					bInsideProceduralVolume = DesiredInstance.ProceduralVolumeBodyInstance->OverlapTest(Hit.ImpactPoint, FQuat::Identity, FCollisionShape::MakeSphere(1.f));	//make sphere of 1cm radius to test if we're in the procedural volume
+				}
 			}
 		}
 			
@@ -1991,7 +1999,7 @@ bool AInstancedFoliageActor::FoliageTrace(const UWorld* InWorld, FHitResult& Out
 		if (Hit.Component.IsValid() && Hit.Component->GetComponentLevel())
 		{
 			OutHit = Hit;
-			return true;
+			return (DesiredInstance.PlacementMode != EFoliagePlacementMode::Procedural) || bInsideProceduralVolume;
 		}
 	}
 
diff --git a/Engine/Source/Runtime/Foliage/Private/ProceduralFoliageActor.cpp b/Engine/Source/Runtime/Foliage/Private/ProceduralFoliageActor.cpp
--- a/Engine/Source/Runtime/Foliage/Private/ProceduralFoliageActor.cpp
+++ b/Engine/Source/Runtime/Foliage/Private/ProceduralFoliageActor.cpp
@@ -8,5 +8,11 @@ AProceduralFoliageActor::AProceduralFoliageActor(const FObjectInitializer& Objec
 : Super(ObjectInitializer)
 {
 	ProceduralComponent = ObjectInitializer.CreateDefaultSubobject<UProceduralFoliageComponent>(this, TEXT("ProceduralFoliageComponent"));
-	RootComponent = ProceduralComponent;
+	ProceduralComponent->SpawningVolume = this;
+
+	if (UBrushComponent* BrushComponent = GetBrushComponent())
+	{
+		BrushComponent->SetCollisionObjectType(ECC_WorldStatic);
+		BrushComponent->SetCollisionResponseToAllChannels(ECR_Ignore);
+	}
 }
diff --git a/Engine/Source/Runtime/Foliage/Private/ProceduralFoliageComponent.cpp b/Engine/Source/Runtime/Foliage/Private/ProceduralFoliageComponent.cpp
--- a/Engine/Source/Runtime/Foliage/Private/ProceduralFoliageComponent.cpp
+++ b/Engine/Source/Runtime/Foliage/Private/ProceduralFoliageComponent.cpp
@@ -14,11 +14,11 @@
 UProceduralFoliageComponent::UProceduralFoliageComponent(const FObjectInitializer& ObjectInitializer)
 : Super(ObjectInitializer)
 {
-	TilesX = 1;
-	TilesY = 1;
 	Overlap = 0.f;
-	HalfHeight = 10000.f;
 	ProceduralGuid = FGuid::NewGuid();
+#if WITH_EDITORONLY_DATA
+	bHideDebugTiles = true;
+#endif
 }
 
 void CopyTileInstances(const UProceduralFoliageTile* FromTile, UProceduralFoliageTile* ToTile, const FBox2D& InnerLocalAABB, const FTransform& ToLocalTM, const float Overlap)
@@ -47,11 +47,48 @@ FBox2D GetTileRegion(const int32 X, const int32 Y, const int32 CountX, const int
 
 	return Region;
 }
+void UProceduralFoliageComponent::GetTilesLayout(int32& MinXIdx, int32& MinYIdx, int32& NumX, int32& NumY, float& HalfHeight) const
+{
+	if(UBrushComponent* Brush = SpawningVolume->GetBrushComponent())
+	{
+		const FVector MinPosition = Brush->Bounds.GetBox().Min;
+		const FVector MaxPosition = Brush->Bounds.GetBox().Max;
+
+		//we want to find the bottom left tile that contains this MinPosition
+		MinXIdx = FMath::FloorToInt(MinPosition.X / ProceduralFoliage->TileSize);
+		MinYIdx = FMath::FloorToInt(MinPosition.Y / ProceduralFoliage->TileSize);
+
+		const int32 MaxXIdx = FMath::FloorToInt(MaxPosition.X / ProceduralFoliage->TileSize);
+		const int32 MaxYIdx = FMath::FloorToInt(MaxPosition.Y / ProceduralFoliage->TileSize);
+
+		NumX = (MaxXIdx - MinXIdx) + 1;
+		NumY = (MaxYIdx - MinYIdx) + 1;
+
+		HalfHeight = Brush->Bounds.GetBox().GetExtent().Z;
+	}
+}
+
+FVector UProceduralFoliageComponent::GetWorldPosition() const
+{
+	UBrushComponent* Brush = SpawningVolume->GetBrushComponent();
+	if (Brush == nullptr || ProceduralFoliage == nullptr)
+	{
+		return FVector::ZeroVector;
+	}
+	
+	int32 X1 = 0;
+	int32 Y1 = 0;
+	int32 NumX, NumY;
+	float HalfHeight;
+	GetTilesLayout(X1, Y1, NumX, NumY, HalfHeight);
+
+	return FVector(X1 * ProceduralFoliage->TileSize, Y1 * ProceduralFoliage->TileSize, Brush->Bounds.Origin.Z);
+}
 
 void UProceduralFoliageComponent::SpawnTiles(TArray<FDesiredFoliageInstance>& OutInstances)
 {
 #if WITH_EDITOR
-	if (ProceduralFoliage)
+	if (ProceduralFoliage && SpawningVolume && SpawningVolume->GetBrushComponent())
 	{
 		/** Constants for laying out the overlapping tile grid*/
 		const float InnerTileSize = ProceduralFoliage->TileSize;
@@ -68,30 +105,38 @@ void UProceduralFoliageComponent::SpawnTiles(TArray<FDesiredFoliageInstance>& Ou
 		const FVector2D OverlapY(0.f, Overlap);
 
 		TArray<TFuture< TArray<FDesiredFoliageInstance>* >> Futures;
-		FScopedSlowTask SlowTask(TilesX * TilesY, LOCTEXT("PlaceProceduralFoliage", "Placing ProceduralFoliage..."));
+		
+		int32 XOffset, YOffset, NumTilesX, NumTilesY;
+		float HalfHeight;
+		GetTilesLayout(XOffset, YOffset, NumTilesX, NumTilesY, HalfHeight);
+
+		FScopedSlowTask SlowTask(NumTilesX * NumTilesY, LOCTEXT("PlaceProceduralFoliage", "Placing ProceduralFoliage..."));
 		SlowTask.MakeDialog();
 
-		for (int32 X = 0; X < TilesX; ++X)
+		FBodyInstance* VolumeBodyInstance = SpawningVolume->GetBrushComponent()->GetBodyInstance();
+
+		const FVector WorldPosition = GetWorldPosition();
+
+		for (int32 X = 0; X < NumTilesX; ++X)
 		{
-			for (int32 Y = 0; Y < TilesY; ++Y)
+			for (int32 Y = 0; Y < NumTilesY; ++Y)
 			{
 				
 				//We have to get the tiles and create new one to build on main thread
-				const UProceduralFoliageTile* Tile = ProceduralFoliage->GetRandomTile(X, Y);
-				const UProceduralFoliageTile* RightTile = (X + 1 < TilesX) ? ProceduralFoliage->GetRandomTile(X + 1, Y) : nullptr;
-				const UProceduralFoliageTile* TopTile = (Y + 1 < TilesY) ? ProceduralFoliage->GetRandomTile(X, Y+1) : nullptr;
-				const UProceduralFoliageTile* TopRightTile = (RightTile && TopTile) ? ProceduralFoliage->GetRandomTile(X + 1, Y+1) : nullptr;
+				const UProceduralFoliageTile* Tile = ProceduralFoliage->GetRandomTile(X + XOffset, Y + YOffset);
+				const UProceduralFoliageTile* RightTile = (X + 1 < NumTilesX) ? ProceduralFoliage->GetRandomTile(X + XOffset + 1, Y + YOffset) : nullptr;
+				const UProceduralFoliageTile* TopTile = (Y + 1 < NumTilesY) ? ProceduralFoliage->GetRandomTile(X + XOffset, Y + YOffset+ 1) : nullptr;
+				const UProceduralFoliageTile* TopRightTile = (RightTile && TopTile) ? ProceduralFoliage->GetRandomTile(X + XOffset + 1, Y + YOffset + 1) : nullptr;
 
 				UProceduralFoliageTile* JTile = ProceduralFoliage->CreateTempTile();
 
 				Futures.Add(Async<TArray<FDesiredFoliageInstance>*>(EAsyncExecution::ThreadPool, [=]()
 				{
-					FTransform TileTM = ComponentToWorld;
-					const FVector OrientedOffset = ComponentToWorld.TransformVectorNoScale(FVector(X, Y, 0.f) * FVector(InnerTileSize, InnerTileSize, 0.f));
-					TileTM.AddToTranslation(OrientedOffset);
+					const FVector OrientedOffset = FVector(X, Y, 0.f) * FVector(InnerTileSize, InnerTileSize, 0.f);
+					const FTransform TileTM(OrientedOffset + WorldPosition);
 
 					//copy the inner tile
-					const FBox2D InnerBox = GetTileRegion(X, Y, TilesX, TilesY, InnerTileSize, Overlap);
+					const FBox2D InnerBox = GetTileRegion(X, Y, NumTilesX, NumTilesY, InnerTileSize, Overlap);
 					CopyTileInstances(Tile, JTile, InnerBox, FTransform::Identity, Overlap);
 
 
@@ -121,7 +166,7 @@ void UProceduralFoliageComponent::SpawnTiles(TArray<FDesiredFoliageInstance>& Ou
 
 					TArray<FDesiredFoliageInstance>* DesiredInstances = new TArray<FDesiredFoliageInstance>();
 					JTile->InstancesToArray();
-					JTile->CreateInstancesToSpawn(*DesiredInstances, TileTM, ProceduralGuid, HalfHeight);
+					JTile->CreateInstancesToSpawn(*DesiredInstances, TileTM, ProceduralGuid, HalfHeight, VolumeBodyInstance);
 					JTile->Empty();
 
 					return DesiredInstances;
@@ -132,9 +177,9 @@ void UProceduralFoliageComponent::SpawnTiles(TArray<FDesiredFoliageInstance>& Ou
 		}
 
 		int32 FutureIdx = 0;
-		for (int X = 0; X < TilesX; ++X)
+		for (int X = 0; X < NumTilesX; ++X)
 		{
-			for (int Y = 0; Y < TilesY; ++Y)
+			for (int Y = 0; Y < NumTilesY; ++Y)
 			{
 				TArray<FDesiredFoliageInstance>* DesiredInstances = Futures[FutureIdx++].Get();
 				OutInstances.Append(*DesiredInstances);
diff --git a/Engine/Source/Runtime/Foliage/Private/ProceduralFoliageTile.cpp b/Engine/Source/Runtime/Foliage/Private/ProceduralFoliageTile.cpp
--- a/Engine/Source/Runtime/Foliage/Private/ProceduralFoliageTile.cpp
+++ b/Engine/Source/Runtime/Foliage/Private/ProceduralFoliageTile.cpp
@@ -433,7 +433,7 @@ void UProceduralFoliageTile::BeginDestroy()
 	RemoveInstances();
 }
 
-void UProceduralFoliageTile::CreateInstancesToSpawn(TArray<FDesiredFoliageInstance>& OutInstances, const FTransform& WorldTM, const FGuid& ProceduralGuid, const float HalfHeight) const
+void UProceduralFoliageTile::CreateInstancesToSpawn(TArray<FDesiredFoliageInstance>& OutInstances, const FTransform& WorldTM, const FGuid& ProceduralGuid, const float HalfHeight, const FBodyInstance* VolumeBodyInstance) const
 {
 	const FCollisionQueryParams Params(true);
 	FHitResult Hit;
@@ -451,6 +451,7 @@ void UProceduralFoliageTile::CreateInstancesToSpawn(TArray<FDesiredFoliageInstan
 		DesiredInst->ProceduralGuid = ProceduralGuid;
 		DesiredInst->FoliageType = Instance.Type;
 		DesiredInst->Age = Instance.Age;
+		DesiredInst->ProceduralVolumeBodyInstance = VolumeBodyInstance;
 		DesiredInst->PlacementMode = EFoliagePlacementMode::Procedural;
 	}
 }
diff --git a/Engine/Source/Runtime/Foliage/Public/InstancedFoliage.h b/Engine/Source/Runtime/Foliage/Public/InstancedFoliage.h
--- a/Engine/Source/Runtime/Foliage/Public/InstancedFoliage.h
+++ b/Engine/Source/Runtime/Foliage/Public/InstancedFoliage.h
@@ -378,6 +378,7 @@ struct FDesiredFoliageInstance
 	FQuat Rotation;
 	float TraceRadius;
 	float Age;
+	const struct FBodyInstance* ProceduralVolumeBodyInstance;
 	EFoliagePlacementMode::Type PlacementMode;
 };
 
diff --git a/Engine/Source/Runtime/Foliage/Public/ProceduralFoliageActor.h b/Engine/Source/Runtime/Foliage/Public/ProceduralFoliageActor.h
--- a/Engine/Source/Runtime/Foliage/Public/ProceduralFoliageActor.h
+++ b/Engine/Source/Runtime/Foliage/Public/ProceduralFoliageActor.h
@@ -7,7 +7,7 @@
 class UProceduralFoliageComponent;
 
 UCLASS()
-class FOLIAGE_API AProceduralFoliageActor: public AActor
+class FOLIAGE_API AProceduralFoliageActor: public AVolume
 {
 	GENERATED_UCLASS_BODY()
 
diff --git a/Engine/Source/Runtime/Foliage/Public/ProceduralFoliageComponent.h b/Engine/Source/Runtime/Foliage/Public/ProceduralFoliageComponent.h
--- a/Engine/Source/Runtime/Foliage/Public/ProceduralFoliageComponent.h
+++ b/Engine/Source/Runtime/Foliage/Public/ProceduralFoliageComponent.h
@@ -10,22 +10,10 @@ class AProceduralFoliageLevelInfo;
 struct FDesiredFoliageInstance;
 
 UCLASS(BlueprintType)
-class FOLIAGE_API UProceduralFoliageComponent : public USceneComponent
+class FOLIAGE_API UProceduralFoliageComponent : public UActorComponent
 {
 	GENERATED_UCLASS_BODY()
 
-	/** The number of ProceduralFoliage tiles to place along the x-axis */
-	UPROPERTY(Category = "ProceduralFoliage", BlueprintReadWrite, EditAnywhere)
-	int32 TilesX;
-
-	/** The number of ProceduralFoliage tiles to place along the y-axis */
-	UPROPERTY(Category = "ProceduralFoliage", BlueprintReadWrite, EditAnywhere)
-	int32 TilesY;
-
-	/** The extent of the volume above and below. Instances will only be placed on geometry within this volume*/
-	UPROPERTY(Category = "ProceduralFoliage", BlueprintReadWrite, EditAnywhere)
-	float HalfHeight;
-
 	/** The overlap in Cms between two tiles. */
 	UPROPERTY(Category = "ProceduralFoliage", BlueprintReadWrite, EditAnywhere)
 	float Overlap;
@@ -33,10 +21,22 @@ class FOLIAGE_API UProceduralFoliageComponent : public USceneComponent
 	UPROPERTY(Category = "ProceduralFoliage", BlueprintReadWrite, EditAnywhere, meta=(DisplayName="ProceduralFoliage Asset") )
 	UProceduralFoliage* ProceduralFoliage;
 
+#if WITH_EDITORONLY_DATA
+	/** The overlap in Cms between two tiles. */
+	UPROPERTY(Category = "ProceduralFoliage", BlueprintReadWrite, EditAnywhere)
+	bool bHideDebugTiles;
+#endif
+
 	void SpawnProceduralContent(TArray<FDesiredFoliageInstance>& OutFoliageInstances);
 	void RemoveProceduralContent();
 	const FGuid& GetProceduralGuid() const { return ProceduralGuid;  }
 
+	UPROPERTY()
+	AVolume* SpawningVolume;
+
+	FVector GetWorldPosition() const;
+	void GetTilesLayout(int32& MinXIdx, int32& MinYIdx, int32& NumXTiles, int32& NumYTiles, float& HalfHeight) const;
+
 private:
 	void SpawnTiles(TArray<FDesiredFoliageInstance>& OutFoliageInstances);
 
diff --git a/Engine/Source/Runtime/Foliage/Public/ProceduralFoliageTile.h b/Engine/Source/Runtime/Foliage/Public/ProceduralFoliageTile.h
--- a/Engine/Source/Runtime/Foliage/Public/ProceduralFoliageTile.h
+++ b/Engine/Source/Runtime/Foliage/Public/ProceduralFoliageTile.h
@@ -27,7 +27,7 @@ class FOLIAGE_API UProceduralFoliageTile : public UObject
 
 	/** Takes a local AABB and finds all the instances in it. The instances are returned in the space local to the AABB. That is, an instance in the bottom left AABB will have a coordinate of (0,0) */
 	void GetInstancesInAABB(const FBox2D& LocalAABB, TArray<FProceduralFoliageInstance* >& OutInstances, bool bOnTheBorder) const;
-	void CreateInstancesToSpawn(TArray<FDesiredFoliageInstance>& OutInstances, const FTransform& WorldTM, const FGuid& ProceduralGuid, const float HalfHeight) const;
+	void CreateInstancesToSpawn(TArray<FDesiredFoliageInstance>& OutInstances, const FTransform& WorldTM, const FGuid& ProceduralGuid, const float HalfHeight, const FBodyInstance* VolumeBodyInstance) const;
 	void AddInstances(const TArray<FProceduralFoliageInstance*>& NewInstances, const FTransform& LocalTM, const FBox2D& InnerLocalAABB);
 
 	virtual SIZE_T GetResourceSize(EResourceSizeMode::Type Mode) override;
