Procedural foliage can now be found in content browser
UE-8157

diff --git a/Engine/Source/Runtime/Foliage/Private/ProceduralFoliageActor.cpp b/Engine/Source/Runtime/Foliage/Private/ProceduralFoliageActor.cpp
--- a/Engine/Source/Runtime/Foliage/Private/ProceduralFoliageActor.cpp
+++ b/Engine/Source/Runtime/Foliage/Private/ProceduralFoliageActor.cpp
@@ -16,3 +16,16 @@ AProceduralFoliageActor::AProceduralFoliageActor(const FObjectInitializer& Objec
 		BrushComponent->SetCollisionResponseToAllChannels(ECR_Ignore);
 	}
 }
+
+#if WITH_EDITOR
+
+bool AProceduralFoliageActor::GetReferencedContentObjects(TArray<UObject*>& Objects) const
+{
+	if (ProceduralComponent && ProceduralComponent->ProceduralFoliage)
+	{
+		Objects.Add(ProceduralComponent->ProceduralFoliage);
+	}
+	return true;
+}
+
+#endif
\ No newline at end of file
diff --git a/Engine/Source/Runtime/Foliage/Public/ProceduralFoliageActor.h b/Engine/Source/Runtime/Foliage/Public/ProceduralFoliageActor.h
--- a/Engine/Source/Runtime/Foliage/Public/ProceduralFoliageActor.h
+++ b/Engine/Source/Runtime/Foliage/Public/ProceduralFoliageActor.h
@@ -13,4 +13,8 @@ class FOLIAGE_API AProceduralFoliageActor: public AVolume
 
 	UPROPERTY(Category = ProceduralFoliageTile, VisibleAnywhere, BlueprintReadOnly)
 	UProceduralFoliageComponent* ProceduralComponent;
+
+#if WITH_EDITOR
+	virtual bool GetReferencedContentObjects(TArray<UObject*>& Objects) const override;
+#endif
 };
\ No newline at end of file
