Clamp for possible negative random number
UE-8158

diff --git a/Engine/Source/Runtime/Foliage/Private/ProceduralFoliage.cpp b/Engine/Source/Runtime/Foliage/Private/ProceduralFoliage.cpp
--- a/Engine/Source/Runtime/Foliage/Private/ProceduralFoliage.cpp
+++ b/Engine/Source/Runtime/Foliage/Private/ProceduralFoliage.cpp
@@ -140,7 +140,7 @@ const UProceduralFoliageTile* UProceduralFoliage::GetRandomTile(int32 X, int32 Y
 		HashStream.Initialize(Y);
 		const float YRand = HashStream.FRand();
 		const int32 RandomNumber = (RAND_MAX * XRand / (YRand + 0.01f));
-		const int32 Idx = RandomNumber % PrecomputedTiles.Num();
+		const int32 Idx = FMath::Clamp(RandomNumber % PrecomputedTiles.Num(), 0, PrecomputedTiles.Num() - 1);
 		return PrecomputedTiles[Idx];
 	}
 
