Clearning a procedural foliage type marks simulation as dirty
UE-8138

diff --git a/Engine/Source/Runtime/Foliage/Private/ProceduralFoliage.cpp b/Engine/Source/Runtime/Foliage/Private/ProceduralFoliage.cpp
--- a/Engine/Source/Runtime/Foliage/Private/ProceduralFoliage.cpp
+++ b/Engine/Source/Runtime/Foliage/Private/ProceduralFoliage.cpp
@@ -49,6 +49,10 @@ void UProceduralFoliage::CreateProceduralFoliageInstances()
 			TypeData.TypeInstance = ConstructObject<UFoliageType_InstancedStaticMesh>(TypeData.Type, this);
 			TypeData.ChangeCount = TypeData.TypeInstance->ChangeCount;
 		}
+		else
+		{
+			TypeData.TypeInstance = nullptr;
+		}
 	}
 }
 
@@ -99,25 +103,21 @@ void UProceduralFoliage::Serialize(FArchive& Ar)
 
 bool UProceduralFoliage::AnyDirty() const
 {
-	bool bDirty = false;
+	bool bDirty = bNeedsSimulation;
 
-	if( Types.Num() > 0 )
-	{
-		bDirty = bNeedsSimulation;
 	for (const FProceduralFoliageTypeData& TypeData : Types)
 	{
 		if (TypeData.Type)
 		{
 			if (TypeData.TypeInstance && TypeData.Type->GetDefaultObject<UFoliageType_InstancedStaticMesh>()->ChangeCount != TypeData.ChangeCount)
 			{
-					bDirty = true;
-					break;
+				bDirty = true;
+				break;
 			}
 			else if (!TypeData.TypeInstance)
 			{
-					bDirty = true;
-					break;
-				}
+				bDirty = true;
+				break;
 			}
 		}
 	}
