Fix non-unity build

diff --git a/Engine/Source/Runtime/Foliage/Private/ProceduralFoliageActor.cpp b/Engine/Source/Runtime/Foliage/Private/ProceduralFoliageActor.cpp
--- a/Engine/Source/Runtime/Foliage/Private/ProceduralFoliageActor.cpp
+++ b/Engine/Source/Runtime/Foliage/Private/ProceduralFoliageActor.cpp
@@ -3,6 +3,7 @@
 #include "FoliagePrivate.h"
 #include "ProceduralFoliageActor.h"
 #include "ProceduralFoliageComponent.h"
+#include "ProceduralFoliage.h"
 
 AProceduralFoliageActor::AProceduralFoliageActor(const FObjectInitializer& ObjectInitializer)
 : Super(ObjectInitializer)
