Add cancel button for procedural foliage

UE-8149

diff --git a/Engine/Source/Editor/FoliageEdit/Private/ProceduralFoliageComponentDetails.cpp b/Engine/Source/Editor/FoliageEdit/Private/ProceduralFoliageComponentDetails.cpp
--- a/Engine/Source/Editor/FoliageEdit/Private/ProceduralFoliageComponentDetails.cpp
+++ b/Engine/Source/Editor/FoliageEdit/Private/ProceduralFoliageComponentDetails.cpp
@@ -71,12 +71,13 @@ FReply FProceduralFoliageComponentDetails::OnResimulateClicked()
 			if( !UniqueProceduralFoliages.Contains( Component->ProceduralFoliage ) )
 			{
 				UniqueProceduralFoliages.Add( Component->ProceduralFoliage );
-				Component->ProceduralFoliage->Simulate();
 			}
 
 			TArray <FDesiredFoliageInstance> DesiredFoliageInstances;
-			Component->SpawnProceduralContent(DesiredFoliageInstances);
-			FEdModeFoliage::AddInstances(Component->GetWorld(), DesiredFoliageInstances);
+			if (Component->SpawnProceduralContent(DesiredFoliageInstances))
+			{
+				FEdModeFoliage::AddInstances(Component->GetWorld(), DesiredFoliageInstances);
+			}
 		}
 	}
 	return FReply::Handled();
diff --git a/Engine/Source/Runtime/Foliage/Private/ProceduralFoliage.cpp b/Engine/Source/Runtime/Foliage/Private/ProceduralFoliage.cpp
--- a/Engine/Source/Runtime/Foliage/Private/ProceduralFoliage.cpp
+++ b/Engine/Source/Runtime/Foliage/Private/ProceduralFoliage.cpp
@@ -32,14 +32,6 @@ UProceduralFoliageTile* UProceduralFoliage::CreateTempTile()
 	return TmpTile;
 }
 
-void UProceduralFoliage::GenerateTile(const int32 NumSteps)
-{
-	UProceduralFoliageTile* NewTile = NewObject<UProceduralFoliageTile>(this);
-	PrecomputedTiles.Add(NewTile);
-
-	NewTile->Simulate(this, GetRandomNumber(), NumSteps);
-}
-
 void UProceduralFoliage::CreateProceduralFoliageInstances()
 {
 	for(FProceduralFoliageTypeData& TypeData : Types)
@@ -125,13 +117,16 @@ bool UProceduralFoliage::AnyDirty() const
 	return bDirty;
 }
 
-const UProceduralFoliageTile* UProceduralFoliage::GetRandomTile(int32 X, int32 Y)
+void UProceduralFoliage::SimulateIfNeeded()
 {
 	if (AnyDirty())
 	{
 		Simulate();
 	}
+}
 
+const UProceduralFoliageTile* UProceduralFoliage::GetRandomTile(int32 X, int32 Y)
+{
 	if (PrecomputedTiles.Num())
 	{
 		FRandomStream HashStream;	//using this as a hash function
@@ -152,6 +147,8 @@ void UProceduralFoliage::Simulate(int32 NumSteps)
 	RandomStream.Initialize(RandomSeed);
 	CreateProceduralFoliageInstances();
 
+	LastCancel.Increment();
+
 	PrecomputedTiles.Empty();
 	TArray<TFuture< UProceduralFoliageTile* >> Futures;
 
@@ -159,25 +156,47 @@ void UProceduralFoliage::Simulate(int32 NumSteps)
 	{
 		UProceduralFoliageTile* NewTile = NewObject<UProceduralFoliageTile>(this);
 		const int32 RandomNumber = GetRandomNumber();
+		const int32 LastCancelInit = LastCancel.GetValue();
 
 		Futures.Add(Async<UProceduralFoliageTile*>(EAsyncExecution::ThreadPool, [=]()
 		{
-			NewTile->Simulate(this, RandomNumber, NumSteps);
+			NewTile->Simulate(this, RandomNumber, NumSteps, LastCancelInit);
 			return NewTile;
 		}));
 	}
 
-	FScopedSlowTask SlowTask(NumUniqueTiles, LOCTEXT("SimulateProceduralFoliage", "Simulate ProceduralFoliage..."));
-	SlowTask.MakeDialog();
+	const FText StatusMessage = LOCTEXT("SimulateProceduralFoliage", "Simulate ProceduralFoliage...");
+	GWarn->BeginSlowTask(StatusMessage, true, true);
+	
+	const int32 TotalTasks = Futures.Num();
+
+	bool bCancelled = false;
 
 	for (int32 FutureIdx = 0; FutureIdx < Futures.Num(); ++FutureIdx)
 	{
-		PrecomputedTiles.Add(Futures[FutureIdx].Get());
-		SlowTask.EnterProgressFrame(1);
+		while (Futures[FutureIdx].WaitFor(FTimespan(0, 0, 0, 0, 100)) == false)		//sleep for 100ms if not ready. Needed so cancel is responsive
+		{
+			GWarn->StatusUpdate(FutureIdx, TotalTasks, LOCTEXT("SimulateProceduralFoliage", "Simulate ProceduralFoliage..."));
+			if (GWarn->ReceivedUserCancel() && bCancelled == false)	//If the user wants to cancel just increment this. Tiles compare against the original count and cancel if needed
+			{
+				LastCancel.Increment();
+				bCancelled = true;
+			}
+		}
+
+		PrecomputedTiles.Add(Futures[FutureIdx].Get());		//Even if we cancel we block until threads have exited safely to ensure memory isn't GCed
 	}
 
-	SetClean();
-	
+	GWarn->EndSlowTask();
+
+	if (bCancelled)
+	{
+		PrecomputedTiles.Empty();
+	}
+	else
+	{
+		SetClean();
+	}
 }
 
 int32 UProceduralFoliage::GetRandomNumber()
diff --git a/Engine/Source/Runtime/Foliage/Private/ProceduralFoliageComponent.cpp b/Engine/Source/Runtime/Foliage/Private/ProceduralFoliageComponent.cpp
--- a/Engine/Source/Runtime/Foliage/Private/ProceduralFoliageComponent.cpp
+++ b/Engine/Source/Runtime/Foliage/Private/ProceduralFoliageComponent.cpp
@@ -85,7 +85,7 @@ FVector UProceduralFoliageComponent::GetWorldPosition() const
 	return FVector(X1 * ProceduralFoliage->TileSize, Y1 * ProceduralFoliage->TileSize, Brush->Bounds.Origin.Z);
 }
 
-void UProceduralFoliageComponent::SpawnTiles(TArray<FDesiredFoliageInstance>& OutInstances)
+bool UProceduralFoliageComponent::SpawnTiles(TArray<FDesiredFoliageInstance>& OutInstances)
 {
 #if WITH_EDITOR
 	if (ProceduralFoliage && SpawningVolume && SpawningVolume->GetBrushComponent())
@@ -110,20 +110,26 @@ void UProceduralFoliageComponent::SpawnTiles(TArray<FDesiredFoliageInstance>& Ou
 		float HalfHeight;
 		GetTilesLayout(XOffset, YOffset, NumTilesX, NumTilesY, HalfHeight);
 
-		FScopedSlowTask SlowTask(NumTilesX * NumTilesY, LOCTEXT("PlaceProceduralFoliage", "Placing ProceduralFoliage..."));
-		SlowTask.MakeDialog();
-
 		FBodyInstance* VolumeBodyInstance = SpawningVolume->GetBrushComponent()->GetBodyInstance();
 
+		FThreadSafeCounter LastCancel;
+		const int32 LastCanelInit = LastCancel.GetValue();
+		FThreadSafeCounter* LastCancelPtr = &LastCancel;
+
 		const FVector WorldPosition = GetWorldPosition();
+		ProceduralFoliage->SimulateIfNeeded();
 
 		for (int32 X = 0; X < NumTilesX; ++X)
 		{
 			for (int32 Y = 0; Y < NumTilesY; ++Y)
 			{
-				
 				//We have to get the tiles and create new one to build on main thread
 				const UProceduralFoliageTile* Tile = ProceduralFoliage->GetRandomTile(X + XOffset, Y + YOffset);
+				if (Tile == nullptr)	//simulation was cancelled or failed
+				{
+					return false;
+				}
+
 				const UProceduralFoliageTile* RightTile = (X + 1 < NumTilesX) ? ProceduralFoliage->GetRandomTile(X + XOffset + 1, Y + YOffset) : nullptr;
 				const UProceduralFoliageTile* TopTile = (Y + 1 < NumTilesY) ? ProceduralFoliage->GetRandomTile(X + XOffset, Y + YOffset+ 1) : nullptr;
 				const UProceduralFoliageTile* TopRightTile = (RightTile && TopTile) ? ProceduralFoliage->GetRandomTile(X + XOffset + 1, Y + YOffset + 1) : nullptr;
@@ -132,6 +138,11 @@ void UProceduralFoliageComponent::SpawnTiles(TArray<FDesiredFoliageInstance>& Ou
 
 				Futures.Add(Async<TArray<FDesiredFoliageInstance>*>(EAsyncExecution::ThreadPool, [=]()
 				{
+					if (LastCancelPtr->GetValue() != LastCanelInit)
+					{
+						return new TArray<FDesiredFoliageInstance>();
+					}
+
 					const FVector OrientedOffset = FVector(X, Y, 0.f) * FVector(InnerTileSize, InnerTileSize, 0.f);
 					const FTransform TileTM(OrientedOffset + WorldPosition);
 
@@ -176,26 +187,64 @@ void UProceduralFoliageComponent::SpawnTiles(TArray<FDesiredFoliageInstance>& Ou
 			}
 		}
 
+		const FText StatusMessage = LOCTEXT("PlaceProceduralFoliage", "Placing ProceduralFoliage...");
+		const FText CancelMessage = LOCTEXT("PlaceProceduralFoliageCancel", "Cancelling ProceduralFoliage...");
+		GWarn->BeginSlowTask(StatusMessage, true, true);
+
+
 		int32 FutureIdx = 0;
+		bool bCancelled = false;
 		for (int X = 0; X < NumTilesX; ++X)
 		{
 			for (int Y = 0; Y < NumTilesY; ++Y)
 			{
+				bool bFirstTime = true;
+				while (Futures[FutureIdx].WaitFor(FTimespan(0, 0, 0, 0, 100)) == false || bFirstTime)
+				{
+					if (GWarn->ReceivedUserCancel() && bCancelled == false)
+					{
+						LastCancel.Increment();
+						bCancelled = true;
+					}
+
+					if (bCancelled)
+					{
+						GWarn->StatusUpdate(Y + X * NumTilesY, NumTilesX * NumTilesY, CancelMessage);
+					}
+					else
+					{
+						GWarn->StatusUpdate(Y + X * NumTilesY, NumTilesX * NumTilesY, StatusMessage);
+					}
+
+					bFirstTime = false;
+				}
+
 				TArray<FDesiredFoliageInstance>* DesiredInstances = Futures[FutureIdx++].Get();
 				OutInstances.Append(*DesiredInstances);
 				delete DesiredInstances;
-				SlowTask.EnterProgressFrame(1);
 			}
 		}
+
+		GWarn->EndSlowTask();
+
+		return !bCancelled;
 	}
 #endif
+	return false;
+
 }
 
-void UProceduralFoliageComponent::SpawnProceduralContent(TArray <FDesiredFoliageInstance>& OutInstances)
+bool UProceduralFoliageComponent::SpawnProceduralContent(TArray <FDesiredFoliageInstance>& OutInstances)
 {
 #if WITH_EDITOR
-	RemoveProceduralContent();
-	SpawnTiles(OutInstances);
+	
+	if (SpawnTiles(OutInstances))
+	{
+		RemoveProceduralContent();
+		return true;
+	}
+
+	return false;
 #endif
 }
 
diff --git a/Engine/Source/Runtime/Foliage/Private/ProceduralFoliageTile.cpp b/Engine/Source/Runtime/Foliage/Private/ProceduralFoliageTile.cpp
--- a/Engine/Source/Runtime/Foliage/Private/ProceduralFoliageTile.cpp
+++ b/Engine/Source/Runtime/Foliage/Private/ProceduralFoliageTile.cpp
@@ -125,6 +125,7 @@ void UProceduralFoliageTile::AgeSeeds()
 	TArray<FProceduralFoliageInstance*> NewSeeds;
 	for (FProceduralFoliageInstance* Instance : Instances)
 	{
+		if (UserCancelled()){ return; }
 		if (Instance->IsAlive())
 		{
 			const UFoliageType_InstancedStaticMesh* Type = Instance->Type;
@@ -158,6 +159,7 @@ void UProceduralFoliageTile::SpreadSeeds(TArray<FProceduralFoliageInstance*>& Ne
 {
 	for (FProceduralFoliageInstance* Inst : Instances)
 	{
+		if (UserCancelled()){ return; }
 		if (Inst->IsAlive() == false)	//The instance has been killed so don't bother spreading seeds. Note this introduces potential indeterminism if the order of instance traversal changes (implementation details of TSet for example)
 		{
 			continue;
@@ -201,6 +203,7 @@ void UProceduralFoliageTile::AddRandomSeeds(TArray<FProceduralFoliageInstance*>&
 
 	for (const FProceduralFoliageTypeData& Data : ProceduralFoliage->GetTypes())
 	{
+		if (UserCancelled()){ return; }
 		const UFoliageType_InstancedStaticMesh* Type = Data.TypeInstance;
 		if (Type && Type->bGrowsInShade == bSimulateShadeGrowth)
 		{
@@ -250,6 +253,7 @@ void UProceduralFoliageTile::AddRandomSeeds(TArray<FProceduralFoliageInstance*>&
 	const int32 LastShadeCastingIndex = InstancesArray.Num() - 1; //when placing shade growth types we want to spawn in shade if possible
 	while (TypesLeftToSeed > 0)
 	{
+		if (UserCancelled()){ return; }
 		TypeIdx = (TypeIdx + 1) % NumTypes;	//keep cycling through the types that we spawn initial seeds for to make sure everyone gets fair chance
 
 		if (const UFoliageType_InstancedStaticMesh* Type = TypesToSeed[TypeIdx])
@@ -368,8 +372,15 @@ void UProceduralFoliageTile::InitSimulation(const UProceduralFoliage* InProcedur
 	Broadphase = FProceduralFoliageBroadphase(ProceduralFoliage->TileSize);
 }
 
+bool UProceduralFoliageTile::UserCancelled() const
+{
+	return ProceduralFoliage->LastCancel.GetValue() != LastCancel;
+}
+
+
 void UProceduralFoliageTile::StepSimulation()
 {
+	if (UserCancelled()){ return; }
 	TArray<FProceduralFoliageInstance*> NewInstances;
 	if (SimulationStep == 0)
 	{
@@ -418,8 +429,9 @@ void UProceduralFoliageTile::StartSimulation(const int32 MaxNumSteps, bool bShad
 	InstancesToArray();
 }
 
-void UProceduralFoliageTile::Simulate(const UProceduralFoliage* InProceduralFoliage, const int32 RandomSeed, const int32 MaxNumSteps)
+void UProceduralFoliageTile::Simulate(const UProceduralFoliage* InProceduralFoliage, const int32 RandomSeed, const int32 MaxNumSteps, const int32 InLastCancel)
 {
+	LastCancel = InLastCancel;
 	InitSimulation(InProceduralFoliage, RandomSeed);
 
 	StartSimulation(MaxNumSteps, false);
diff --git a/Engine/Source/Runtime/Foliage/Public/ProceduralFoliage.h b/Engine/Source/Runtime/Foliage/Public/ProceduralFoliage.h
--- a/Engine/Source/Runtime/Foliage/Public/ProceduralFoliage.h
+++ b/Engine/Source/Runtime/Foliage/Public/ProceduralFoliage.h
@@ -39,6 +39,9 @@ class FOLIAGE_API UProceduralFoliage : public UObject
 
 	UPROPERTY(Category = ProceduralFoliageSimulation, EditAnywhere, BlueprintReadOnly)
 	int32 NumUniqueTiles;
+
+	FThreadSafeCounter LastCancel;
+
 private:
 	UPROPERTY(Category = ProceduralFoliageSimulation, EditAnywhere)
 	TArray<FProceduralFoliageTypeData> Types;
@@ -63,14 +66,16 @@ public:
 
 	virtual void Serialize(FArchive& Ar);
 
+	/** Simulates tiles if needed */
+	void SimulateIfNeeded();
+
 	/** Takes a tile index and returns a random tile associated with that index. */
 	const UProceduralFoliageTile* GetRandomTile(int32 X, int32 Y);
 
 	/** Creates a temporary empty tile with the appropriate settings created for it. */
 	UProceduralFoliageTile* CreateTempTile();
-private:
-	void GenerateTile(const int32 NumSteps);
 
+private:
 	void CreateProceduralFoliageInstances();
 
 	bool AnyDirty() const;
diff --git a/Engine/Source/Runtime/Foliage/Public/ProceduralFoliageComponent.h b/Engine/Source/Runtime/Foliage/Public/ProceduralFoliageComponent.h
--- a/Engine/Source/Runtime/Foliage/Public/ProceduralFoliageComponent.h
+++ b/Engine/Source/Runtime/Foliage/Public/ProceduralFoliageComponent.h
@@ -27,7 +27,7 @@ class FOLIAGE_API UProceduralFoliageComponent : public UActorComponent
 	bool bHideDebugTiles;
 #endif
 
-	void SpawnProceduralContent(TArray<FDesiredFoliageInstance>& OutFoliageInstances);
+	bool SpawnProceduralContent(TArray<FDesiredFoliageInstance>& OutFoliageInstances);
 	void RemoveProceduralContent();
 	const FGuid& GetProceduralGuid() const { return ProceduralGuid;  }
 
@@ -38,7 +38,7 @@ class FOLIAGE_API UProceduralFoliageComponent : public UActorComponent
 	void GetTilesLayout(int32& MinXIdx, int32& MinYIdx, int32& NumXTiles, int32& NumYTiles, float& HalfHeight) const;
 
 private:
-	void SpawnTiles(TArray<FDesiredFoliageInstance>& OutFoliageInstances);
+	bool SpawnTiles(TArray<FDesiredFoliageInstance>& OutFoliageInstances);
 
 private:
 	UPROPERTY()
diff --git a/Engine/Source/Runtime/Foliage/Public/ProceduralFoliageTile.h b/Engine/Source/Runtime/Foliage/Public/ProceduralFoliageTile.h
--- a/Engine/Source/Runtime/Foliage/Public/ProceduralFoliageTile.h
+++ b/Engine/Source/Runtime/Foliage/Public/ProceduralFoliageTile.h
@@ -21,7 +21,7 @@ class FOLIAGE_API UProceduralFoliageTile : public UObject
 	UPROPERTY()
 	const UProceduralFoliage* ProceduralFoliage;
 
-	void Simulate(const UProceduralFoliage* ProceduralFoliage, const int32 RandomSeed, const int32 MaxNumSteps = -1);
+	void Simulate(const UProceduralFoliage* ProceduralFoliage, const int32 RandomSeed, const int32 MaxNumSteps, const int32 InLastCancel);
 	void RemoveInstances();
 	void RemoveInstance(FProceduralFoliageInstance* Inst);
 
@@ -57,6 +57,8 @@ private:
 
 	void MarkPendingRemoval(FProceduralFoliageInstance* ToRemove);
 	void FlushPendingRemovals();
+
+	bool UserCancelled() const;
 	
 private:
 
@@ -72,6 +74,8 @@ private:
 	FRandomStream RandomStream;
 	bool bSimulateShadeGrowth;
 
+	int32 LastCancel;
+
 private:
 	float GetRandomGaussian();
 	FVector GetSeedOffset(const UFoliageType_InstancedStaticMesh* Type, float MinDistance);
