Add some buffer for procedural volume height so that volumes right on top of meshes still spawn instances.

UE-9204

diff --git a/Engine/Source/Runtime/Foliage/Private/ProceduralFoliageTile.cpp b/Engine/Source/Runtime/Foliage/Private/ProceduralFoliageTile.cpp
--- a/Engine/Source/Runtime/Foliage/Private/ProceduralFoliageTile.cpp
+++ b/Engine/Source/Runtime/Foliage/Private/ProceduralFoliageTile.cpp
@@ -456,7 +456,7 @@ void UProceduralFoliageTile::CreateInstancesToSpawn(TArray<FDesiredFoliageInstan
 		FVector StartRay = Instance.Location + WorldTM.GetLocation();
 		StartRay.Z += HalfHeight;
 		FVector EndRay = StartRay;
-		EndRay.Z -= HalfHeight*2.f;
+		EndRay.Z -= (HalfHeight*2.f + 10.f);	//add 10cm to bottom position of raycast. This is needed because volume is usually placed directly on geometry and then you get precision issues
 
 		FDesiredFoliageInstance* DesiredInst = new (OutInstances)FDesiredFoliageInstance(StartRay, EndRay, Instance.GetMaxRadius());
 		DesiredInst->Rotation = Instance.Rotation;
