Turn sweeps back on in procedural foliage. Needed to make sure overlapping procedural volumes don't have overlapping instances

diff --git a/Engine/Source/Runtime/Foliage/Private/InstancedFoliage.cpp b/Engine/Source/Runtime/Foliage/Private/InstancedFoliage.cpp
--- a/Engine/Source/Runtime/Foliage/Private/InstancedFoliage.cpp
+++ b/Engine/Source/Runtime/Foliage/Private/InstancedFoliage.cpp
@@ -1991,7 +1991,11 @@ bool AInstancedFoliageActor::FoliageTrace(const UWorld* InWorld, FHitResult& Out
 	TArray<FHitResult> Hits;
 
 	bool bInsideProceduralVolume = false;
-	InWorld->LineTraceMulti(Hits, StartTrace, DesiredInstance.EndTrace, QueryParams, FCollisionObjectQueryParams(ECC_WorldStatic));
+	FCollisionShape SphereShape;
+	SphereShape.SetSphere(DesiredInstance.TraceRadius);
+	InWorld->SweepMulti(Hits, StartTrace, DesiredInstance.EndTrace, FQuat::Identity, SphereShape, QueryParams, FCollisionObjectQueryParams(ECC_WorldStatic));
+
+
 	for (const FHitResult& Hit : Hits)
 	{
 		if (DesiredInstance.PlacementMode == EFoliagePlacementMode::Procedural)
