Combine procedural foliage seed and per species seed for random placement. This way modifying procedural foliage seed still has effect on initial placement

diff --git a/Engine/Source/Runtime/Foliage/Private/ProceduralFoliageTile.cpp b/Engine/Source/Runtime/Foliage/Private/ProceduralFoliageTile.cpp
--- a/Engine/Source/Runtime/Foliage/Private/ProceduralFoliageTile.cpp
+++ b/Engine/Source/Runtime/Foliage/Private/ProceduralFoliageTile.cpp
@@ -217,7 +217,7 @@ void UProceduralFoliageTile::AddRandomSeeds(TArray<FProceduralFoliageInstance*>&
 			}
 
 			{	//save the random stream per type
-				RandomStreamPerType.Add(Type, FRandomStream(Type->DistributionSeed));
+				RandomStreamPerType.Add(Type, FRandomStream(Type->DistributionSeed + ProceduralFoliage->RandomSeed));
 			}
 
 			{	//compute the needed offsets for initial seed variance
