Fix crash when using procedural foliage. AddFoliageType was not returning a mesh info if it found static mesh for the foliage type that is not an asset and not owned by the foliage actor

diff --git a/Engine/Source/Runtime/Foliage/Private/InstancedFoliage.cpp b/Engine/Source/Runtime/Foliage/Private/InstancedFoliage.cpp
--- a/Engine/Source/Runtime/Foliage/Private/InstancedFoliage.cpp
+++ b/Engine/Source/Runtime/Foliage/Private/InstancedFoliage.cpp
@@ -1352,9 +1352,10 @@ UFoliageType* AInstancedFoliageActor::AddFoliageType(const UFoliageType* InType,
 	{
 		// Unique meshes only
 		// Multiple entries for same static mesh can be added using FoliageType as an asset
-		FoliageType = GetSettingsForMesh(FoliageType->GetStaticMesh());
+		FoliageType = GetSettingsForMesh(FoliageType->GetStaticMesh(), &MeshInfo);
 		if (FoliageType == nullptr)
 		{
+			//TODO: Duplicating a null object seems wrong. Am I missing something?
 			FoliageType = DuplicateObject<UFoliageType>(FoliageType, this);
 			MarkPackageDirty();
 			MeshInfo = &FoliageMeshes.Add(FoliageType).Get();
