Fixed typo in foliage settings duplication

diff --git a/Engine/Source/Runtime/Foliage/Private/InstancedFoliage.cpp b/Engine/Source/Runtime/Foliage/Private/InstancedFoliage.cpp
--- a/Engine/Source/Runtime/Foliage/Private/InstancedFoliage.cpp
+++ b/Engine/Source/Runtime/Foliage/Private/InstancedFoliage.cpp
@@ -1355,8 +1355,7 @@ UFoliageType* AInstancedFoliageActor::AddFoliageType(const UFoliageType* InType,
 		FoliageType = GetSettingsForMesh(FoliageType->GetStaticMesh(), &MeshInfo);
 		if (FoliageType == nullptr)
 		{
-			//TODO: Duplicating a null object seems wrong. Am I missing something?
-			FoliageType = DuplicateObject<UFoliageType>(FoliageType, this);
+			FoliageType = DuplicateObject<UFoliageType>(InType, this);
 			MarkPackageDirty();
 			MeshInfo = &FoliageMeshes.Add(FoliageType).Get();
 		}
