'Culling Distance' and 'Collision Settings' are missing in foliage tool #UE-9716

diff --git a/Engine/Source/Runtime/Foliage/Private/InstancedFoliage.cpp b/Engine/Source/Runtime/Foliage/Private/InstancedFoliage.cpp
--- a/Engine/Source/Runtime/Foliage/Private/InstancedFoliage.cpp
+++ b/Engine/Source/Runtime/Foliage/Private/InstancedFoliage.cpp
@@ -267,6 +267,8 @@ UFoliageType::UFoliageType(const FObjectInitializer& ObjectInitializer)
 	Height.Max = 262144.0f;
 	ZOffset.Min = 0.0f;
 	ZOffset.Max = 0.0f;
+	CullDistance.Min = 0;
+	CullDistance.Max = 0;
 	MinimumLayerWeight = 0.5f;
 	DisplayOrder = 0;
 	IsSelected = false;
diff --git a/Engine/Source/Runtime/Foliage/Public/FoliageType.h b/Engine/Source/Runtime/Foliage/Public/FoliageType.h
--- a/Engine/Source/Runtime/Foliage/Public/FoliageType.h
+++ b/Engine/Source/Runtime/Foliage/Public/FoliageType.h
@@ -156,7 +156,7 @@ class UFoliageType : public UObject
 	 * The distance where instances will begin to fade out if using a PerInstanceFadeAmount material node. 0 disables.
 	 * When the entire cluster is beyond this distance, the cluster is completely culled and not rendered at all.
 	 */
-	UPROPERTY(EditAnywhere, Category = General)
+	UPROPERTY(EditAnywhere, Category = General, meta = (Subcategory = "Culling"))
 	FInt32Interval CullDistance;
 	
 	UPROPERTY()
@@ -215,7 +215,7 @@ class UFoliageType : public UObject
 	int32 OverriddenLightMapRes;
 
 	/** Custom collision for foliage */
-	UPROPERTY(EditAnywhere, Category = General, meta = (HideObjectType = true))
+	UPROPERTY(EditAnywhere, Category = General, meta = (Subcategory = "Collision", HideObjectType = true))
 	struct FBodyInstance BodyInstance;
 
 	/* Gets/Sets the mesh associated with this FoliageType */
