Change collision with world to use channel instead of object type. Object type WorldStatic is already used by volumes which would prevent painting.

diff --git a/Engine/Source/Runtime/Foliage/Private/InstancedFoliage.cpp b/Engine/Source/Runtime/Foliage/Private/InstancedFoliage.cpp
--- a/Engine/Source/Runtime/Foliage/Private/InstancedFoliage.cpp
+++ b/Engine/Source/Runtime/Foliage/Private/InstancedFoliage.cpp
@@ -2128,7 +2128,7 @@ bool AInstancedFoliageActor::CheckCollisionWithWorld(const UWorld* InWorld, cons
 			FBoxSphereBounds WorldBound = ShrinkBound.TransformBy(InstTransform);
 			//::DrawDebugBox(World, WorldBound.Origin, WorldBound.BoxExtent, FColor::Red, true, 10.f);
 			static FName NAME_FoliageCollisionWithWorld = FName(TEXT("FoliageCollisionWithWorld"));
-			if (InWorld->OverlapAnyTest(WorldBound.Origin, FQuat(Inst.Rotation), FCollisionShape::MakeBox(ShrinkBound.BoxExtent * Inst.DrawScale3D * Settings->CollisionScale), FCollisionQueryParams(NAME_FoliageCollisionWithWorld, false), FCollisionObjectQueryParams(ECC_WorldStatic)))
+			if (InWorld->OverlapBlockingTest(WorldBound.Origin, FQuat(Inst.Rotation), ECC_WorldStatic, FCollisionShape::MakeBox(ShrinkBound.BoxExtent * Inst.DrawScale3D * Settings->CollisionScale), FCollisionQueryParams(NAME_FoliageCollisionWithWorld, false)))
 			{
 				return false;
 			}
