Change tiles of procedural component to ignore overlap parts. This makes it easier to have falloff

diff --git a/Engine/Source/Runtime/Foliage/Private/ProceduralFoliageComponent.cpp b/Engine/Source/Runtime/Foliage/Private/ProceduralFoliageComponent.cpp
--- a/Engine/Source/Runtime/Foliage/Private/ProceduralFoliageComponent.cpp
+++ b/Engine/Source/Runtime/Foliage/Private/ProceduralFoliageComponent.cpp
@@ -51,8 +51,8 @@ void UProceduralFoliageComponent::GetTilesLayout(int32& MinXIdx, int32& MinYIdx,
 {
 	if(UBrushComponent* Brush = SpawningVolume->GetBrushComponent())
 	{
-		const FVector MinPosition = Brush->Bounds.GetBox().Min;
-		const FVector MaxPosition = Brush->Bounds.GetBox().Max;
+		const FVector MinPosition = Brush->Bounds.GetBox().Min + Overlap;
+		const FVector MaxPosition = Brush->Bounds.GetBox().Max - Overlap;
 
 		//we want to find the bottom left tile that contains this MinPosition
 		MinXIdx = FMath::FloorToInt(MinPosition.X / ProceduralFoliage->TileSize);
