Force procedural resimulation since tiles are no longer saved

diff --git a/Engine/Source/Runtime/Foliage/Private/ProceduralFoliage.cpp b/Engine/Source/Runtime/Foliage/Private/ProceduralFoliage.cpp
--- a/Engine/Source/Runtime/Foliage/Private/ProceduralFoliage.cpp
+++ b/Engine/Source/Runtime/Foliage/Private/ProceduralFoliage.cpp
@@ -119,10 +119,11 @@ bool UProceduralFoliage::AnyDirty() const
 
 void UProceduralFoliage::SimulateIfNeeded()
 {
-	if (AnyDirty())
+	//if (AnyDirty())	@todo: for now we must force simulation every time since precomputed tiles are weak pointers
 	{
 		Simulate();
 	}
+
 }
 
 const UProceduralFoliageTile* UProceduralFoliage::GetRandomTile(int32 X, int32 Y)
