Prevent foliage actor copy, paste, duplicate, move to level operations UE-9969

diff --git a/Engine/Source/Runtime/Foliage/Private/InstancedFoliage.cpp b/Engine/Source/Runtime/Foliage/Private/InstancedFoliage.cpp
--- a/Engine/Source/Runtime/Foliage/Private/InstancedFoliage.cpp
+++ b/Engine/Source/Runtime/Foliage/Private/InstancedFoliage.cpp
@@ -1633,6 +1633,16 @@ void AInstancedFoliageActor::PostEditUndo()
 	}
 }
 
+bool AInstancedFoliageActor::ShouldExport()
+{
+	return false;
+}
+
+bool AInstancedFoliageActor::ShouldImport(FString* ActorPropString, bool IsMovingLevel)
+{
+	return false;
+}
+
 void AInstancedFoliageActor::ApplySelectionToComponents(bool bApply)
 {
 	for (auto& MeshPair : FoliageMeshes)
@@ -1886,7 +1896,19 @@ void AInstancedFoliageActor::PostLoad()
 {
 	Super::PostLoad();
 
-	GetLevel()->InstancedFoliageActor = this;
+	ULevel* OwningLevel = GetLevel();
+	if (!OwningLevel->InstancedFoliageActor.IsValid())
+	{
+		OwningLevel->InstancedFoliageActor = this;
+	}
+	else
+	{
+		// Warn that there is more than one foliage actor in the scene
+		UE_LOG(LogInstancedFoliage, Warning, TEXT("Level %s: has more than one instanced foliage actor: %s, %s"), 
+			*OwningLevel->GetOutermost()->GetName(), 
+			*OwningLevel->InstancedFoliageActor->GetName(),
+			*this->GetName());
+	}
 
 #if WITH_EDITOR
 	if (GIsEditor)
@@ -1979,13 +2001,6 @@ void AInstancedFoliageActor::PostLoad()
 #endif
 }
 
-void AInstancedFoliageActor::PostRegisterAllComponents()
-{
-	Super::PostRegisterAllComponents();
-	//
-	GetLevel()->InstancedFoliageActor = this;
-}
-
 #if WITH_EDITOR
 
 void AInstancedFoliageActor::NotifyFoliageTypeChanged(UFoliageType* FoliageType)
diff --git a/Engine/Source/Runtime/Foliage/Public/InstancedFoliageActor.h b/Engine/Source/Runtime/Foliage/Public/InstancedFoliageActor.h
--- a/Engine/Source/Runtime/Foliage/Public/InstancedFoliageActor.h
+++ b/Engine/Source/Runtime/Foliage/Public/InstancedFoliageActor.h
@@ -39,7 +39,6 @@ public:
 	virtual void PostInitProperties() override;
 	virtual void BeginDestroy() override;
 	virtual void PostLoad() override;
-	virtual void PostRegisterAllComponents() override;
 	static void AddReferencedObjects(UObject* InThis, FReferenceCollector& Collector);
 	// End UObject interface. 
 
@@ -88,6 +87,8 @@ public:
 
 #if WITH_EDITOR
 	virtual void PostEditUndo() override;
+	virtual bool ShouldExport() override;
+	virtual bool ShouldImport(FString* ActorPropString, bool IsMovingLevel) override;
 
 	// Called in response to BSP rebuilds to migrate foliage from obsolete to new components.
 	FOLIAGE_API void MapRebuild();
