Conditionally register foliage components CL# 2440988
Merging using UE4-GDC2015-To-UE4

diff --git a/Engine/Source/Runtime/Foliage/Private/InstancedFoliage.cpp b/Engine/Source/Runtime/Foliage/Private/InstancedFoliage.cpp
--- a/Engine/Source/Runtime/Foliage/Private/InstancedFoliage.cpp
+++ b/Engine/Source/Runtime/Foliage/Private/InstancedFoliage.cpp
@@ -620,7 +620,11 @@ void FFoliageMeshInfo::AddInstance(AInstancedFoliageActor* InIFA, const UFoliage
 		UpdateComponentSettings(InSettings);
 
 		Component->AttachTo(InIFA->GetRootComponent());
-		Component->RegisterComponent();
+		
+		if (InIFA->GetRootComponent()->IsRegistered())
+		{
+			Component->RegisterComponent();
+		}
 
 		// Use only instance translation as a component transform
 		Component->SetWorldTransform(InIFA->GetRootComponent()->ComponentToWorld);
@@ -846,7 +850,11 @@ void FFoliageMeshInfo::ReallocateClusters(AInstancedFoliageActor* InIFA, UFoliag
 		}
 	}
 
-	InIFA->RegisterAllComponents();
+	// Register components if foliage level is visible
+	if (InIFA->GetLevel()->bIsVisible)
+	{
+		InIFA->RegisterAllComponents();
+	}
 }
 
 void FFoliageMeshInfo::ReapplyInstancesToComponent()
