Unregistering only target component in FFoliageMeshInfo::ReallocateClusters CL# 2441052
Merging using UE4-GDC2015-To-UE4

diff --git a/Engine/Source/Runtime/Foliage/Private/InstancedFoliage.cpp b/Engine/Source/Runtime/Foliage/Private/InstancedFoliage.cpp
--- a/Engine/Source/Runtime/Foliage/Private/InstancedFoliage.cpp
+++ b/Engine/Source/Runtime/Foliage/Private/InstancedFoliage.cpp
@@ -821,12 +821,9 @@ void FFoliageMeshInfo::RemoveFromBaseHash(int32 InstanceIndex)
 // Destroy existing clusters and reassign all instances to new clusters
 void FFoliageMeshInfo::ReallocateClusters(AInstancedFoliageActor* InIFA, UFoliageType* InSettings)
 {
-	// Detach all components
-
-	InIFA->UnregisterAllComponents();
-
 	if (Component != nullptr)
 	{
+		Component->UnregisterComponent();
 		Component->bAutoRegister = false;
 		Component = nullptr;
 	}
@@ -849,12 +846,6 @@ void FFoliageMeshInfo::ReallocateClusters(AInstancedFoliageActor* InIFA, UFoliag
 			AddInstance(InIFA, InSettings, Instance);
 		}
 	}
-
-	// Register components if foliage level is visible
-	if (InIFA->GetLevel()->bIsVisible)
-	{
-		InIFA->RegisterAllComponents();
-	}
 }
 
 void FFoliageMeshInfo::ReapplyInstancesToComponent()
