Merging using UE4-GDC2015-To-UE4
FoliageEdit + CL# 2442692

diff --git a/Engine/Source/Runtime/Engine/Classes/Components/PrimitiveComponent.h b/Engine/Source/Runtime/Engine/Classes/Components/PrimitiveComponent.h
--- a/Engine/Source/Runtime/Engine/Classes/Components/PrimitiveComponent.h
+++ b/Engine/Source/Runtime/Engine/Classes/Components/PrimitiveComponent.h
@@ -1661,6 +1661,8 @@ public:
 
 	FORCEINLINE EHasCustomNavigableGeometry::Type HasCustomNavigableGeometry() const { return bHasCustomNavigableGeometry; }
 
+	void SetCustomNavigableGeometry(const EHasCustomNavigableGeometry::Type InType);
+
 	/** Collects custom navigable geometry of component.
 	 *	@return true if regular navigable geometry exporting should be run as well */
 	virtual bool DoCustomNavigableGeometryExport(struct FNavigableGeometryExport* GeomExport) const { return true; }
diff --git a/Engine/Source/Runtime/Engine/Private/PrimitiveComponent.cpp b/Engine/Source/Runtime/Engine/Private/PrimitiveComponent.cpp
--- a/Engine/Source/Runtime/Engine/Private/PrimitiveComponent.cpp
+++ b/Engine/Source/Runtime/Engine/Private/PrimitiveComponent.cpp
@@ -2644,5 +2644,10 @@ bool UPrimitiveComponent::CanEditSimulatePhysics()
 	return GetBodySetup() != nullptr;
 }
 
+void UPrimitiveComponent::SetCustomNavigableGeometry(const EHasCustomNavigableGeometry::Type InType)
+{
+	bHasCustomNavigableGeometry = InType;
+}
+
 #undef LOCTEXT_NAMESPACE
 
diff --git a/Engine/Source/Runtime/Foliage/Private/FoliageStatistics.cpp b/Engine/Source/Runtime/Foliage/Private/FoliageStatistics.cpp
--- a/Engine/Source/Runtime/Foliage/Private/FoliageStatistics.cpp
+++ b/Engine/Source/Runtime/Foliage/Private/FoliageStatistics.cpp
@@ -17,11 +17,11 @@ int32 UFoliageStatistics::FoliageOverlappingSphereCount(UObject* WorldContextObj
 	const FSphere Sphere(CenterPosition, Radius);
 
 	int32 Count = 0;
-	for (FConstLevelIterator LvlItr = World->GetLevelIterator(); LvlItr; ++LvlItr)
+
+	for (TActorIterator<AInstancedFoliageActor> It(World); It; ++It)
 	{
-		ULevel* Level = *LvlItr;
-		AInstancedFoliageActor* IFA = AInstancedFoliageActor::GetInstancedFoliageActorForLevel(Level);
-		if (IFA)
+		AInstancedFoliageActor* IFA = *It;
+		if (!IFA->IsPendingKill())
 		{
 			const UFoliageType* FoliageType = IFA->GetSettingsForMesh(Mesh);
 			Count += IFA->GetOverlappingSphereCount(FoliageType, Sphere);
@@ -36,17 +36,18 @@ int32 UFoliageStatistics::FoliageOverlappingBoxCount(UObject* WorldContextObject
 	UWorld* World = GEngine->GetWorldFromContextObject(WorldContextObject);
 
 	int32 Count = 0;
-	for(FConstLevelIterator LvlItr = World->GetLevelIterator() ; LvlItr ; ++LvlItr)
+
+	for (TActorIterator<AInstancedFoliageActor> It(World); It; ++It)
 	{
-		ULevel* Level = *LvlItr;
-		AInstancedFoliageActor* IFA = AInstancedFoliageActor::GetInstancedFoliageActorForLevel(Level);
-		if (IFA)
+		AInstancedFoliageActor* IFA = *It;
+		if (!IFA->IsPendingKill())
 		{
 			const UFoliageType* Type = IFA->GetSettingsForMesh(StaticMesh);
 			Count += IFA->GetOverlappingBoxCount(Type, Box);
 		}
 	}
 
+
 	return Count;
 }
 
@@ -54,14 +55,29 @@ void UFoliageStatistics::FoliageOverlappingBoxTransforms(UObject* WorldContextOb
 {
 	UWorld* World = GEngine->GetWorldFromContextObject(WorldContextObject);
 
-	for(FConstLevelIterator LvlItr = World->GetLevelIterator() ; LvlItr ; ++LvlItr)
+	for (TActorIterator<AInstancedFoliageActor> It(World); It; ++It)
 	{
-		ULevel* Level = *LvlItr;
-		AInstancedFoliageActor* IFA = AInstancedFoliageActor::GetInstancedFoliageActorForLevel(Level);
-		if (IFA)
+		AInstancedFoliageActor* IFA = *It;
+		if (!IFA->IsPendingKill())
 		{
 			const UFoliageType* Type = IFA->GetSettingsForMesh(StaticMesh);
 			IFA->GetOverlappingBoxTransforms(Type, Box, OutTransforms);
 		}
 	}
 }
+
+void UFoliageStatistics::FoliageOverlappingMeshCounts_Debug(UObject* WorldContextObject, FVector CenterPosition, float Radius, TMap<UStaticMesh*, int32>& OutMeshCounts)
+{
+	UWorld* const World = GEngine->GetWorldFromContextObject(WorldContextObject);
+	const FSphere Sphere(CenterPosition, Radius);
+
+	OutMeshCounts.Reset();
+	for (TActorIterator<AInstancedFoliageActor> It(World); It; ++It)
+	{
+		AInstancedFoliageActor* IFA = *It;
+		if (!IFA->IsPendingKill())
+		{
+			IFA->GetOverlappingMeshCounts(Sphere, OutMeshCounts);
+		}
+	}
+}
\ No newline at end of file
diff --git a/Engine/Source/Runtime/Foliage/Private/FoliageStatistics.h b/Engine/Source/Runtime/Foliage/Private/FoliageStatistics.h
--- a/Engine/Source/Runtime/Foliage/Private/FoliageStatistics.h
+++ b/Engine/Source/Runtime/Foliage/Private/FoliageStatistics.h
@@ -36,5 +36,14 @@ class FOLIAGE_API UFoliageStatistics : public UBlueprintFunctionLibrary
 	 *	@param OutTransforms Array to populate with transforms
 	 */
 	static void FoliageOverlappingBoxTransforms(UObject* WorldContextObject, const UStaticMesh* StaticMesh, FBox Box, TArray<FTransform>& OutTransforms);
+
+	/** 
+	 *  DEBUG FUNCTION: This is not fast, use only for debug/development.
+	 *  Gets an instance count for each unique mesh type overlapping the given sphere.
+	 *  @param	CenterPosition	The center position of the sphere
+	 *  @param	Radius			The radius of the sphere.
+	 *	@param OutMeshCounts	Map of Meshes to instance counts
+	 */
+	static void FoliageOverlappingMeshCounts_Debug(UObject* WorldContextObject, FVector CenterPosition, float Radius, TMap<UStaticMesh*, int32>& OutMeshCounts);
 };
 
diff --git a/Engine/Source/Runtime/Foliage/Private/InstancedFoliage.cpp b/Engine/Source/Runtime/Foliage/Private/InstancedFoliage.cpp
--- a/Engine/Source/Runtime/Foliage/Private/InstancedFoliage.cpp
+++ b/Engine/Source/Runtime/Foliage/Private/InstancedFoliage.cpp
@@ -594,6 +594,8 @@ void FFoliageMeshInfo::UpdateComponentSettings(const UFoliageType* InSettings)
 		Component->OverriddenLightMapRes = InSettings->OverriddenLightMapRes;
 
 		Component->BodyInstance.CopyBodyInstancePropertiesFrom(&InSettings->BodyInstance);
+
+		Component->SetCustomNavigableGeometry(InSettings->CustomNavigableGeometry);
 	}
 }
 
@@ -1055,7 +1057,7 @@ int32 AInstancedFoliageActor::GetOverlappingSphereCount(const UFoliageType* Foli
 {
 	if (const FFoliageMeshInfo* MeshInfo = FindMesh(FoliageType))
 	{
-		if (MeshInfo->Component)
+		if(MeshInfo->Component && MeshInfo->Component->IsTreeFullyBuilt())
 		{
 			return MeshInfo->Component->GetOverlappingSphereCount(Sphere);
 		}
@@ -1069,7 +1071,7 @@ int32 AInstancedFoliageActor::GetOverlappingBoxCount(const UFoliageType* Foliage
 {
 	if(const FFoliageMeshInfo* MeshInfo = FindMesh(FoliageType))
 	{
-		if(MeshInfo->Component)
+		if(MeshInfo->Component && MeshInfo->Component->IsTreeFullyBuilt())
 		{
 			return MeshInfo->Component->GetOverlappingBoxCount(Box);
 		}
@@ -1083,13 +1085,33 @@ void AInstancedFoliageActor::GetOverlappingBoxTransforms(const UFoliageType* Fol
 {
 	if(const FFoliageMeshInfo* MeshInfo = FindMesh(FoliageType))
 	{
-		if(MeshInfo->Component)
+		if(MeshInfo->Component && MeshInfo->Component->IsTreeFullyBuilt())
 		{
 			MeshInfo->Component->GetOverlappingBoxTransforms(Box, OutTransforms);
 		}
 	}
 }
 
+void AInstancedFoliageActor::GetOverlappingMeshCounts(const FSphere& Sphere, TMap<UStaticMesh*, int32>& OutCounts) const
+{
+	for (auto& MeshPair : FoliageMeshes)
+	{
+		FFoliageMeshInfo const* MeshInfo = &*MeshPair.Value;
+
+		if (MeshInfo && MeshInfo->Component)
+		{
+			int32 const Count = MeshInfo->Component->GetOverlappingSphereCount(Sphere);
+			if (Count > 0)
+			{
+				UStaticMesh* const Mesh = MeshInfo->Component->StaticMesh;
+				int32& StoredCount = OutCounts.FindOrAdd(Mesh);
+				StoredCount += Count;
+			}
+		}
+	}
+}
+
+
 
 UFoliageType* AInstancedFoliageActor::GetSettingsForMesh(const UStaticMesh* InMesh, FFoliageMeshInfo** OutMeshInfo)
 {
@@ -2008,12 +2030,12 @@ void AInstancedFoliageActor::PostLoad()
 				MeshInfo.ReallocateClusters(this, MeshPair.Key);
 			}
 
-			// Update foliage components if the foliage settings object was changed while the level was not loaded.
+			// Update foliage component settings if the foliage settings object was changed while the level was not loaded.
 			if (MeshInfo.FoliageTypeUpdateGuid != FoliageType->UpdateGuid)
 			{
 				if (MeshInfo.FoliageTypeUpdateGuid.IsValid())
 				{
-					MeshInfo.ReallocateClusters(this, MeshPair.Key);
+					MeshInfo.UpdateComponentSettings(FoliageType);
 				}
 				MeshInfo.FoliageTypeUpdateGuid = FoliageType->UpdateGuid;
 			}
@@ -2027,7 +2049,7 @@ void AInstancedFoliageActor::PostLoad()
 				MeshInfo.InstanceHash->InsertInstance(MeshInfo.Instances[InstanceIdx].Location, InstanceIdx);
 			}
 	
-			// Convert to Heirarchical foliage
+			// Convert to Hierarchical foliage
 			if (GetLinkerCustomVersion(FFoliageCustomVersion::GUID) < FFoliageCustomVersion::FoliageUsingHierarchicalISMC)
 			{
 				MeshInfo.ReallocateClusters(this, MeshPair.Key);
@@ -2119,13 +2141,15 @@ void AInstancedFoliageActor::ApplyWorldOffset(const FVector& InOffset, bool bWor
 {
 	Super::ApplyWorldOffset(InOffset, bWorldShift);
 
-	if (GIsEditor)
+#if WITH_EDITORONLY_DATA	
+	UWorld* World = GetWorld();
+
+	if (GIsEditor && World && !World->IsGameWorld())
 	{
 		for (auto& MeshPair : FoliageMeshes)
 		{
 			FFoliageMeshInfo& MeshInfo = *MeshPair.Value;
 
-#if WITH_EDITORONLY_DATA
 			InstanceBaseCache.UpdateInstanceBaseCachedTransforms();
 			
 			MeshInfo.InstanceHash->Empty();
@@ -2136,9 +2160,9 @@ void AInstancedFoliageActor::ApplyWorldOffset(const FVector& InOffset, bool bWor
 				// Rehash instance location
 				MeshInfo.InstanceHash->InsertInstance(Instance.Location, InstanceIdx);
 			}
-#endif
 		}
 	}
+#endif
 }
 
 bool AInstancedFoliageActor::FoliageTrace(const UWorld* InWorld, FHitResult& OutHit, const FDesiredFoliageInstance& DesiredInstance, FName InTraceTag, bool InbReturnFaceIndex)
diff --git a/Engine/Source/Runtime/Foliage/Public/FoliageType.h b/Engine/Source/Runtime/Foliage/Public/FoliageType.h
--- a/Engine/Source/Runtime/Foliage/Public/FoliageType.h
+++ b/Engine/Source/Runtime/Foliage/Public/FoliageType.h
@@ -14,6 +14,10 @@ enum FoliageVertexColorMask
 	FOLIAGEVERTEXCOLORMASK_Alpha	UMETA(DisplayName="Alpha"),
 };
 
+namespace EHasCustomNavigableGeometry
+{
+	enum Type;
+}
 UENUM()
 enum class EFoliageScaling : uint8
 {
@@ -29,6 +33,7 @@ enum class EFoliageScaling : uint8
 	LockYZ
 };
 
+
 UCLASS(hidecategories = Object, editinlinenew, MinimalAPI, BlueprintType, Blueprintable)
 class UFoliageType : public UObject
 {
@@ -245,6 +250,10 @@ class UFoliageType : public UObject
 	UPROPERTY(EditAnywhere, Category=InstanceSettings, meta = (HideObjectType = true))
 	struct FBodyInstance BodyInstance;
 
+	/** Force navmesh */
+	UPROPERTY(EditAnywhere, Category=InstanceSettings, meta = (HideObjectType = true))
+	TEnumAsByte<EHasCustomNavigableGeometry::Type> CustomNavigableGeometry;
+
 	UPROPERTY()
 	int32 DisplayOrder;
 
diff --git a/Engine/Source/Runtime/Foliage/Public/InstancedFoliageActor.h b/Engine/Source/Runtime/Foliage/Public/InstancedFoliageActor.h
--- a/Engine/Source/Runtime/Foliage/Public/InstancedFoliageActor.h
+++ b/Engine/Source/Runtime/Foliage/Public/InstancedFoliageActor.h
@@ -46,6 +46,7 @@ public:
 	// we don't want to have our components automatically destroyed by the Blueprint code
 	virtual void RerunConstructionScripts() override {}
 	virtual void ApplyWorldOffset(const FVector& InOffset, bool bWorldShift) override;
+	virtual bool IsLevelBoundsRelevant() const override { return false; }
 	// End AActor interface.
 
 
@@ -58,6 +59,9 @@ public:
 	FOLIAGE_API int32 GetOverlappingBoxCount(const UFoliageType* FoliageType, const FBox& Box) const;
 	// Finds all instances in the provided box and get their transforms
 	FOLIAGE_API void GetOverlappingBoxTransforms(const UFoliageType* FoliageType, const FBox& Box, TArray<FTransform>& OutTransforms) const;
+	// Perf Warnin: potentially slow! Dev-only use recommended.
+	// Returns list of meshes and counts for all nearby instances. OutCounts accumulates between runs.
+	FOLIAGE_API void GetOverlappingMeshCounts(const FSphere& Sphere, TMap<UStaticMesh*, int32>& OutCounts) const;
 
 	// Finds a mesh entry
 	FOLIAGE_API FFoliageMeshInfo* FindMesh(const UFoliageType* InType);
