Fixed: Deleting landscape components should remove attached foliage
UE-11275

diff --git a/Engine/Source/Editor/LandscapeEditor/Private/LandscapeEdModeComponentTools.cpp b/Engine/Source/Editor/LandscapeEditor/Private/LandscapeEdModeComponentTools.cpp
--- a/Engine/Source/Editor/LandscapeEditor/Private/LandscapeEdModeComponentTools.cpp
+++ b/Engine/Source/Editor/LandscapeEditor/Private/LandscapeEdModeComponentTools.cpp
@@ -987,6 +987,16 @@ public:
 				ALandscape::SplitHeightmap(Component, false);
 			}
 
+			// Remove attached foliage
+			for (TSet<ULandscapeComponent*>::TIterator It(SelectedComponents); It; ++It)
+			{
+				ULandscapeHeightfieldCollisionComponent* CollisionComp = (*It)->CollisionComponent.Get();
+				if (CollisionComp)
+				{
+					AInstancedFoliageActor::DeleteInstancesForComponent(ViewportClient->GetWorld(), CollisionComp);
+				}
+			}
+
 			TArray<FIntPoint> DeletedNeighborKeys;
 			// Check which ones are need for height map change
 			for (TSet<ULandscapeComponent*>::TIterator It(SelectedComponents); It; ++It)
diff --git a/Engine/Source/Runtime/Foliage/Private/InstancedFoliage.cpp b/Engine/Source/Runtime/Foliage/Private/InstancedFoliage.cpp
--- a/Engine/Source/Runtime/Foliage/Private/InstancedFoliage.cpp
+++ b/Engine/Source/Runtime/Foliage/Private/InstancedFoliage.cpp
@@ -1253,6 +1253,16 @@ void AInstancedFoliageActor::DeleteInstancesForComponent(UActorComponent* InComp
 	}
 }
 
+void AInstancedFoliageActor::DeleteInstancesForComponent(UWorld* InWorld, UActorComponent* InComponent)
+{
+	for (TActorIterator<AInstancedFoliageActor> It(InWorld); It; ++It)
+	{
+		AInstancedFoliageActor* IFA = (*It);
+		IFA->Modify();
+		IFA->DeleteInstancesForComponent(InComponent);
+	}
+}
+
 void AInstancedFoliageActor::DeleteInstancesForProceduralFoliageComponent(const UProceduralFoliageComponent* ProceduralFoliageComponent)
 {
 	const FGuid& ProceduralGuid = ProceduralFoliageComponent->GetProceduralGuid();
@@ -1310,27 +1320,26 @@ void AInstancedFoliageActor::MoveInstancesForComponentToCurrentLevel(UActorCompo
 
 void AInstancedFoliageActor::MoveInstancesToNewComponent(UPrimitiveComponent* InOldComponent, UPrimitiveComponent* InNewComponent)
 {
-	AInstancedFoliageActor* NewIFA = AInstancedFoliageActor::GetInstancedFoliageActorForLevel(InNewComponent->GetTypedOuter<ULevel>(), true);
+	AInstancedFoliageActor* TargetIFA = AInstancedFoliageActor::GetInstancedFoliageActorForLevel(InNewComponent->GetTypedOuter<ULevel>(), true);
 
 	const auto OldBaseId = this->InstanceBaseCache.GetInstanceBaseId(InOldComponent);
-	const auto NewBaseId = NewIFA->InstanceBaseCache.AddInstanceBaseId(InNewComponent);
+	if (OldBaseId == FFoliageInstanceBaseCache::InvalidBaseId)
+	{
+		// This foliage actor has no any instances with specified base
+		return;
+	}
+
+	const auto NewBaseId = TargetIFA->InstanceBaseCache.AddInstanceBaseId(InNewComponent);
 
 	for (auto& MeshPair : FoliageMeshes)
 	{
 		FFoliageMeshInfo& MeshInfo = *MeshPair.Value;
-		UFoliageType* FoliageType = MeshPair.Key;
-
-		// Duplicate the foliage type if it's not shared
-		if (FoliageType->GetOutermost() == GetOutermost())
-		{
-			FoliageType = (UFoliageType*)StaticDuplicateObject(FoliageType, NewIFA, nullptr, RF_AllFlags & ~(RF_Standalone | RF_Public));
-		}
-
+				
 		TSet<int32> InstanceSet;
 		if (MeshInfo.ComponentHash.RemoveAndCopyValue(OldBaseId, InstanceSet) && InstanceSet.Num())
 		{
 			// For same FoliageActor can just remap the instances, otherwise we have to do a more complex move
-			if (NewIFA == this)
+			if (TargetIFA == this)
 			{
 				// Update the instances
 				for (int32 InstanceIndex : InstanceSet)
@@ -1340,18 +1349,18 @@ void AInstancedFoliageActor::MoveInstancesToNewComponent(UPrimitiveComponent* In
 				
 				// Update the hash
 				MeshInfo.ComponentHash.Add(NewBaseId, MoveTemp(InstanceSet));
-				
 			}
 			else
 			{
-				FFoliageMeshInfo* NewMeshInfo = NewIFA->FindOrAddMesh(FoliageType);
+				FFoliageMeshInfo* TargetMeshInfo = nullptr;
+				UFoliageType* TargetFoliageType = TargetIFA->AddFoliageType(MeshPair.Key, &TargetMeshInfo);
 
 				// Add the foliage to the new level
 				for (int32 InstanceIndex : InstanceSet)
 				{
 					FFoliageInstance NewInstance = MeshInfo.Instances[InstanceIndex];
 					NewInstance.BaseId = NewBaseId;
-					NewMeshInfo->AddInstance(NewIFA, FoliageType, NewInstance);
+					TargetMeshInfo->AddInstance(TargetIFA, TargetFoliageType, NewInstance);
 				}
 
 				// Remove from old level
diff --git a/Engine/Source/Runtime/Foliage/Public/InstancedFoliageActor.h b/Engine/Source/Runtime/Foliage/Public/InstancedFoliageActor.h
--- a/Engine/Source/Runtime/Foliage/Public/InstancedFoliageActor.h
+++ b/Engine/Source/Runtime/Foliage/Public/InstancedFoliageActor.h
@@ -117,6 +117,9 @@ public:
 	FOLIAGE_API void DeleteInstancesForComponent(UActorComponent* InComponent);
 	FOLIAGE_API void DeleteInstancesForComponent(UActorComponent* InComponent, const UFoliageType* InFoliageType);
 
+	// Deletes the instances attached to a component, traverses all foliage actors in the world
+	static FOLIAGE_API void DeleteInstancesForComponent(UWorld* InWorld, UActorComponent* InComponent);
+
 	// Deletes the instances spawned by a procedural component
 	void DeleteInstancesForProceduralFoliageComponent(const UProceduralFoliageComponent* ProceduralComponent);
 
