Fix procedural foliage not spawning on BSP. UE-9204

diff --git a/Engine/Source/Runtime/Foliage/Private/InstancedFoliage.cpp b/Engine/Source/Runtime/Foliage/Private/InstancedFoliage.cpp
--- a/Engine/Source/Runtime/Foliage/Private/InstancedFoliage.cpp
+++ b/Engine/Source/Runtime/Foliage/Private/InstancedFoliage.cpp
@@ -2231,29 +2231,26 @@ bool AInstancedFoliageActor::FoliageTrace(const UWorld* InWorld, FHitResult& Out
 	{
 		if (DesiredInstance.PlacementMode == EFoliagePlacementMode::Procedural)
 		{
-			if (Hit.Actor.IsValid())	//if we hit the ProceduralFoliage blocking volume don't spawn instance
+			if (const AProceduralFoliageBlockingVolume* ProceduralFoliageBlockingVolume = Cast<AProceduralFoliageBlockingVolume>(Hit.Actor.Get()))
 			{
-				if (const AProceduralFoliageBlockingVolume* ProceduralFoliageBlockingVolume = Cast<AProceduralFoliageBlockingVolume>(Hit.Actor.Get()))
-				{
-					const AProceduralFoliageActor* ProceduralFoliageActor = ProceduralFoliageBlockingVolume->ProceduralFoliageActor;
-					if (ProceduralFoliageActor == nullptr || ProceduralFoliageActor->ProceduralComponent == nullptr || ProceduralFoliageActor->ProceduralComponent->GetProceduralGuid() == DesiredInstance.ProceduralGuid)
-					{
-						return false;
-					}
-				}
-				else if(Cast<AInstancedFoliageActor>(Hit.Actor.Get()))
+				const AProceduralFoliageActor* ProceduralFoliageActor = ProceduralFoliageBlockingVolume->ProceduralFoliageActor;
+				if (ProceduralFoliageActor == nullptr || ProceduralFoliageActor->ProceduralComponent == nullptr || ProceduralFoliageActor->ProceduralComponent->GetProceduralGuid() == DesiredInstance.ProceduralGuid)
 				{
 					return false;
 				}
-				else if (Cast<AProceduralFoliageActor>(Hit.Actor.Get()))	//we never want to collide with our spawning volume
-				{
-					continue;
-				}
+			}
+			else if(Cast<AInstancedFoliageActor>(Hit.Actor.Get()))
+			{
+				return false;
+			}
+			else if (Cast<AProceduralFoliageActor>(Hit.Actor.Get()))	//we never want to collide with our spawning volume
+			{
+				continue;
+			}
 
-				if (bInsideProceduralVolume == false)
-				{
-					bInsideProceduralVolume = DesiredInstance.ProceduralVolumeBodyInstance->OverlapTest(Hit.ImpactPoint, FQuat::Identity, FCollisionShape::MakeSphere(1.f));	//make sphere of 1cm radius to test if we're in the procedural volume
-				}
+			if (bInsideProceduralVolume == false)
+			{
+				bInsideProceduralVolume = DesiredInstance.ProceduralVolumeBodyInstance->OverlapTest(Hit.ImpactPoint, FQuat::Identity, FCollisionShape::MakeSphere(1.f));	//make sphere of 1cm radius to test if we're in the procedural volume
 			}
 		}
 			
