Turn overlap events off for procedural foliage volume. This is not needed and causes a hitch during streaming since it overlaps with thousands of instances.

diff --git a/Engine/Source/Runtime/Foliage/Private/ProceduralFoliageActor.cpp b/Engine/Source/Runtime/Foliage/Private/ProceduralFoliageActor.cpp
--- a/Engine/Source/Runtime/Foliage/Private/ProceduralFoliageActor.cpp
+++ b/Engine/Source/Runtime/Foliage/Private/ProceduralFoliageActor.cpp
@@ -15,6 +15,8 @@ AProceduralFoliageActor::AProceduralFoliageActor(const FObjectInitializer& Objec
 	{
 		BrushComponent->SetCollisionObjectType(ECC_WorldStatic);
 		BrushComponent->SetCollisionResponseToAllChannels(ECR_Ignore);
+		//This is important because the volume overlaps with all procedural foliage, which means during streaming we'll get a huge hitch for UpdateOverlaps
+		BrushComponent->bGenerateOverlapEvents = false;	//We don't care about overlaps.
 	}
 }
 
