Fixed problem where undoing the deletion of the InstancedFoliageActor didn't restore the component correctly. UE-9956

diff --git a/Engine/Source/Runtime/Foliage/Private/InstancedFoliage.cpp b/Engine/Source/Runtime/Foliage/Private/InstancedFoliage.cpp
--- a/Engine/Source/Runtime/Foliage/Private/InstancedFoliage.cpp
+++ b/Engine/Source/Runtime/Foliage/Private/InstancedFoliage.cpp
@@ -854,6 +854,9 @@ void FFoliageMeshInfo::ReapplyInstancesToComponent()
 {
 	if (Component)
 	{
+		// clear the transactional flag if it was set prior to deleting the actor
+		Component->ClearFlags(RF_Transactional);
+
 		Component->UnregisterComponent();
 		Component->ClearInstances();
 
@@ -1723,6 +1726,24 @@ bool AInstancedFoliageActor::HasSelectedInstances() const
 	return false;
 }
 
+void AInstancedFoliageActor::Destroyed()
+{
+	for (auto& MeshPair : FoliageMeshes)
+	{
+		UHierarchicalInstancedStaticMeshComponent* Component = MeshPair.Value->Component;
+
+		if (Component)
+		{
+			Component->ClearInstances();
+			// Save the component's PendingKill flag to restore the component if the delete is undone.
+			Component->SetFlags(RF_Transactional);
+			Component->Modify();
+		}
+	}
+
+	Super::Destroyed();
+}
+
 void AInstancedFoliageActor::PostEditUndo()
 {
 	Super::PostEditUndo();
