Missing .h file for CL 2476471

diff --git a/Engine/Source/Runtime/Foliage/Public/InstancedFoliageActor.h b/Engine/Source/Runtime/Foliage/Public/InstancedFoliageActor.h
--- a/Engine/Source/Runtime/Foliage/Public/InstancedFoliageActor.h
+++ b/Engine/Source/Runtime/Foliage/Public/InstancedFoliageActor.h
@@ -47,6 +47,7 @@ public:
 	virtual void RerunConstructionScripts() override {}
 	virtual void ApplyWorldOffset(const FVector& InOffset, bool bWorldShift) override;
 	virtual bool IsLevelBoundsRelevant() const override { return false; }
+	virtual void Destroyed() override;
 	// End AActor interface.
 
 
