Fix non-editor build for JackP

diff --git a/Engine/Source/Runtime/Foliage/Public/InstancedFoliageActor.h b/Engine/Source/Runtime/Foliage/Public/InstancedFoliageActor.h
--- a/Engine/Source/Runtime/Foliage/Public/InstancedFoliageActor.h
+++ b/Engine/Source/Runtime/Foliage/Public/InstancedFoliageActor.h
@@ -47,7 +47,9 @@ public:
 	virtual void RerunConstructionScripts() override {}
 	virtual void ApplyWorldOffset(const FVector& InOffset, bool bWorldShift) override;
 	virtual bool IsLevelBoundsRelevant() const override { return false; }
+#if WITH_EDITOR
 	virtual void Destroyed() override;
+#endif
 	// End AActor interface.
 
 
