Fixed typo that was prevening foliage to be moved to another level along with landscape components

diff --git a/Engine/Source/Runtime/Foliage/Private/InstancedFoliage.cpp b/Engine/Source/Runtime/Foliage/Private/InstancedFoliage.cpp
--- a/Engine/Source/Runtime/Foliage/Private/InstancedFoliage.cpp
+++ b/Engine/Source/Runtime/Foliage/Private/InstancedFoliage.cpp
@@ -1290,7 +1290,7 @@ void AInstancedFoliageActor::DeleteInstancesForProceduralFoliageComponent(const
 
 void AInstancedFoliageActor::MoveInstancesForComponentToCurrentLevel(UActorComponent* InComponent)
 {
-	if (HasFoliageAttached(InComponent))
+	if (!HasFoliageAttached(InComponent))
 	{
 		// Quit early if there are no foliage instances painted on this component
 		return;
