Changed foliage InstanceHash to use map of sets instead multimap, as it significantly speed-ups level loading with big amount of foliage

diff --git a/Engine/Source/Runtime/Foliage/Public/InstancedFoliage.h b/Engine/Source/Runtime/Foliage/Public/InstancedFoliage.h
--- a/Engine/Source/Runtime/Foliage/Public/InstancedFoliage.h
+++ b/Engine/Source/Runtime/Foliage/Public/InstancedFoliage.h
@@ -237,7 +237,7 @@ struct FFoliageInstanceHash
 {
 private:
 	const int32 HashCellBits;
-	TMultiMap<uint64, int32> CellMap;
+	TMap<uint64, TSet<int32>> CellMap;
 
 	uint64 MakeKey(int32 CellX, int32 CellY)
 	{
@@ -249,8 +249,6 @@ private:
 		return  MakeKey(FMath::FloorToInt(Location.X) >> HashCellBits, FMath::FloorToInt(Location.Y) >> HashCellBits);
 	}
 
-	// Locality map
-	//TMap<uint64, TSet<int32>> CellMap;
 public:
 	FFoliageInstanceHash(int32 InHashCellBits = FOLIAGE_HASH_CELL_BITS)
 	:	HashCellBits(InHashCellBits)
@@ -259,15 +257,15 @@ public:
 	void InsertInstance(const FVector& InstanceLocation, int32 InstanceIndex)
 	{
 		uint64 Key = MakeKey(InstanceLocation);
-
-		CellMap.AddUnique(Key, InstanceIndex);
+			
+		CellMap.FindOrAdd(Key).Add(InstanceIndex);
 	}
 
 	void RemoveInstance(const FVector& InstanceLocation, int32 InstanceIndex)
 	{
 		uint64 Key = MakeKey(InstanceLocation);
 		
-		int32 RemoveCount = CellMap.RemoveSingle(Key, InstanceIndex);
+		int32 RemoveCount = CellMap.FindChecked(Key).Remove(InstanceIndex);
 		check(RemoveCount == 1);
 	}
 
@@ -283,7 +281,11 @@ public:
 			for (int32 x = MinX; x <= MaxX; x++)
 			{
 				uint64 Key = MakeKey(x, y);
-				CellMap.MultiFind(Key, OutInstanceIndices);
+				auto* SetPtr = CellMap.Find(Key);
+				if(SetPtr)
+				{
+					OutInstanceIndices.Append(SetPtr->Array());
+				}
 			}
 		}
 	}
@@ -298,7 +300,13 @@ public:
 #if UE_BUILD_DEBUG
 	void CheckInstanceCount(int32 InCount)
 	{
-		check(CellMap.Num() == InCount);
+		int32 HashCount = 0;
+		for (const auto& Pair : CellMap)
+		{
+			HashCount+= Pair.Value.Num();
+		}
+		
+		check(HashCount == InCount);
 	}
 #endif
 
@@ -309,25 +317,7 @@ public:
 
 	friend FArchive& operator<<(FArchive& Ar, FFoliageInstanceHash& Hash)
 	{
-		if (Ar.UE4Ver() < VER_UE4_FOLIAGE_SETTINGS_TYPE)
-		{
-			Hash.CellMap.Reset();
-
-			TMap<uint64, TSet<int32>> OldCellMap;
-			Ar << OldCellMap;
-			for (auto& CellPair : OldCellMap)
-			{
-				for (int32 Idx : CellPair.Value)
-				{
-					Hash.CellMap.AddUnique(CellPair.Key, Idx);
-				}
-			}
-		}
-		else
-		{
-			Ar << Hash.CellMap;
-		}
-
+		Ar << Hash.CellMap;
 		return Ar;
 	}
 };
