TArray::BulkSerialize supports transactions
Foliage instances use BulkSerialize during transactions, InstanceHash is not serialized anymore and instead regenerated on Undo to speed up 'Modify' operations

diff --git a/Engine/Source/Runtime/Core/Public/Containers/Array.h b/Engine/Source/Runtime/Core/Public/Containers/Array.h
--- a/Engine/Source/Runtime/Core/Public/Containers/Array.h
+++ b/Engine/Source/Runtime/Core/Public/Containers/Array.h
@@ -1211,7 +1211,8 @@ public:
 
 		if (bForcePerElementSerialization
 			|| (Ar.IsSaving()			// if we are saving, we always do the ordinary serialize as a way to make sure it matches up with bulk serialization
-			&& !Ar.IsCooking())		// but cooking is performance critical, so we skip that
+			&& !Ar.IsCooking()			// but cooking and transacting is performance critical, so we skip that
+			&& !Ar.IsTransacting())		
 			|| Ar.IsByteSwapping()		// if we are byteswapping, we need to do that per-element
 			)
 		{
diff --git a/Engine/Source/Runtime/Foliage/Private/InstancedFoliage.cpp b/Engine/Source/Runtime/Foliage/Private/InstancedFoliage.cpp
--- a/Engine/Source/Runtime/Foliage/Private/InstancedFoliage.cpp
+++ b/Engine/Source/Runtime/Foliage/Private/InstancedFoliage.cpp
@@ -224,7 +224,14 @@ FArchive& operator<<(FArchive& Ar, FFoliageMeshInfo& MeshInfo)
 #if WITH_EDITORONLY_DATA
 	if (!Ar.ArIsFilterEditorOnly && !(Ar.GetPortFlags() & PPF_DuplicateForPIE))
 	{
-		Ar << MeshInfo.Instances;
+		if (Ar.IsTransacting())
+		{
+			MeshInfo.Instances.BulkSerialize(Ar);
+		}
+		else
+		{
+			Ar << MeshInfo.Instances;
+		}
 	}
 
 	if (!Ar.ArIsFilterEditorOnly)
@@ -235,7 +242,6 @@ FArchive& operator<<(FArchive& Ar, FFoliageMeshInfo& MeshInfo)
 	// Serialize the transient data for undo.
 	if (Ar.IsTransacting())
 	{
-		Ar << *MeshInfo.InstanceHash;
 		Ar << MeshInfo.ComponentHash;
 		Ar << MeshInfo.SelectedIndices;
 	}
@@ -1753,6 +1759,14 @@ void AInstancedFoliageActor::PostEditUndo()
 	{
 		FFoliageMeshInfo& MeshInfo = *MeshPair.Value;
 		MeshInfo.ReapplyInstancesToComponent();
+
+		// Regenerate instance hash
+		// We regenerate it here instead of saving to transaction buffer to speed up modify operations
+		MeshInfo.InstanceHash->Empty();
+		for (int32 InstanceIdx = 0; InstanceIdx < MeshInfo.Instances.Num(); InstanceIdx++)
+		{
+			MeshInfo.InstanceHash->InsertInstance(MeshInfo.Instances[InstanceIdx].Location, InstanceIdx);
+		}
 	}
 }
 
