Copy foliage settings changes to existing Component instead of recreating it. Only dirty lighting cache when lighting properties are changed.

diff --git a/Engine/Source/Runtime/Foliage/Private/InstancedFoliage.cpp b/Engine/Source/Runtime/Foliage/Private/InstancedFoliage.cpp
--- a/Engine/Source/Runtime/Foliage/Private/InstancedFoliage.cpp
+++ b/Engine/Source/Runtime/Foliage/Private/InstancedFoliage.cpp
@@ -585,23 +585,94 @@ void FFoliageMeshInfo::UpdateComponentSettings(const UFoliageType* InSettings)
 {
 	if (Component)
 	{
-		Component->Mobility = InSettings->bEnableStaticLighting ? EComponentMobility::Static : EComponentMobility::Movable;
-		Component->InstanceStartCullDistance = InSettings->CullDistance.Min;
-		Component->InstanceEndCullDistance = InSettings->CullDistance.Max;
-
-		Component->CastShadow = InSettings->CastShadow;
-		Component->bCastDynamicShadow = InSettings->bCastDynamicShadow;
-		Component->bCastStaticShadow = InSettings->bCastStaticShadow;
-		Component->bAffectDynamicIndirectLighting = InSettings->bAffectDynamicIndirectLighting;
-		Component->bAffectDistanceFieldLighting = InSettings->bAffectDistanceFieldLighting;
-		Component->bCastShadowAsTwoSided = InSettings->bCastShadowAsTwoSided;
-		Component->bReceivesDecals = InSettings->bReceivesDecals;
-		Component->bOverrideLightMapRes = InSettings->bOverrideLightMapRes;
-		Component->OverriddenLightMapRes = InSettings->OverriddenLightMapRes;
+		bool bNeedsMarkRenderStateDirty = false;
+		bool bNeedsInvalidateLightingCache = false;
 
+		EComponentMobility::Type NewMobility = InSettings->bEnableStaticLighting ? EComponentMobility::Static : EComponentMobility::Movable;
+		if (Component->Mobility != NewMobility)
+		{
+			Component->Mobility = NewMobility;
+			bNeedsMarkRenderStateDirty = true;
+			bNeedsInvalidateLightingCache = true;
+		}
+		if (Component->InstanceStartCullDistance != InSettings->CullDistance.Min)
+		{
+			Component->InstanceStartCullDistance = InSettings->CullDistance.Min;
+			bNeedsMarkRenderStateDirty = true;
+		}
+		if (Component->InstanceEndCullDistance != InSettings->CullDistance.Max)
+		{
+			Component->InstanceEndCullDistance = InSettings->CullDistance.Max;
+			bNeedsMarkRenderStateDirty = true;
+		}
+		if (Component->CastShadow != InSettings->CastShadow)
+		{
+			Component->CastShadow = InSettings->CastShadow;
+			bNeedsMarkRenderStateDirty = true;
+			bNeedsInvalidateLightingCache = true;
+		}
+		if (Component->bCastDynamicShadow != InSettings->bCastDynamicShadow)
+		{
+			Component->bCastDynamicShadow = InSettings->bCastDynamicShadow;
+			bNeedsMarkRenderStateDirty = true;
+			bNeedsInvalidateLightingCache = true;
+		}
+		if (Component->bCastStaticShadow != InSettings->bCastStaticShadow)
+		{
+			Component->bCastStaticShadow = InSettings->bCastStaticShadow;
+			bNeedsMarkRenderStateDirty = true;
+			bNeedsInvalidateLightingCache = true;
+		}
+		if (Component->bAffectDynamicIndirectLighting != InSettings->bAffectDynamicIndirectLighting)
+		{
+			Component->bAffectDynamicIndirectLighting = InSettings->bAffectDynamicIndirectLighting;
+			bNeedsMarkRenderStateDirty = true;
+			bNeedsInvalidateLightingCache = true;
+		}
+		if (Component->bAffectDistanceFieldLighting != InSettings->bAffectDistanceFieldLighting)
+		{
+			Component->bAffectDistanceFieldLighting = InSettings->bAffectDistanceFieldLighting;
+			bNeedsMarkRenderStateDirty = true;
+			bNeedsInvalidateLightingCache = true;
+		}
+		if (Component->bCastShadowAsTwoSided != InSettings->bCastShadowAsTwoSided)
+		{
+			Component->bCastShadowAsTwoSided = InSettings->bCastShadowAsTwoSided;
+			bNeedsMarkRenderStateDirty = true;
+			bNeedsInvalidateLightingCache = true;
+		}
+		if (Component->bReceivesDecals != InSettings->bReceivesDecals)
+		{
+			Component->bReceivesDecals = InSettings->bReceivesDecals;
+			bNeedsMarkRenderStateDirty = true;
+			bNeedsInvalidateLightingCache = true;
+		}
+		if (Component->bOverrideLightMapRes != InSettings->bOverrideLightMapRes)
+		{
+			Component->bOverrideLightMapRes = InSettings->bOverrideLightMapRes;
+			bNeedsMarkRenderStateDirty = true;
+			bNeedsInvalidateLightingCache = true;
+		}
+		if (Component->OverriddenLightMapRes != InSettings->OverriddenLightMapRes)
+		{
+			Component->OverriddenLightMapRes = InSettings->OverriddenLightMapRes;
+			bNeedsMarkRenderStateDirty = true;
+			bNeedsInvalidateLightingCache = true;
+		}
+		
 		Component->BodyInstance.CopyBodyInstancePropertiesFrom(&InSettings->BodyInstance);
 
 		Component->SetCustomNavigableGeometry(InSettings->CustomNavigableGeometry);
+
+		if (bNeedsInvalidateLightingCache)
+		{
+			Component->InvalidateLightingCache();
+		}
+
+		if (bNeedsMarkRenderStateDirty)
+		{
+			Component->MarkRenderStateDirty();
+		}
 	}
 }
 
@@ -2160,7 +2231,7 @@ void AInstancedFoliageActor::NotifyFoliageTypeChanged(UFoliageType* FoliageType)
 	FFoliageMeshInfo* MeshInfo = FindMesh(FoliageType);
 	if (MeshInfo)
 	{
-		MeshInfo->ReallocateClusters(this, FoliageType);
+		MeshInfo->UpdateComponentSettings(FoliageType);
 	}
 }
 
