- Renamed UProceduralFoliage to UProceduralFoliageSpawner
- Renamed UProceduralFoliageFactory to UProceduralFoliageSpawnerFactory

diff --git a/Engine/Config/BaseEngine.ini b/Engine/Config/BaseEngine.ini
--- a/Engine/Config/BaseEngine.ini
+++ b/Engine/Config/BaseEngine.ini
@@ -355,6 +355,9 @@ bUseStreamingPause=false
 
 ;Procedural foliage redirects
 +ActiveClassRedirects=(OldClassName="ProceduralFoliageActor",NewClassName="ProceduralFoliageVolume")
++ActiveClassRedirects=(OldClassName="ProceduralFoliage",NewClassName="ProceduralFoliageSpawner")
++TaggedPropertyRedirects=(ClassName="ProceduralFoliageComponent",OldPropertyName="ProceduralFoliage",NewPropertyName="FoliageSpawner")
++TaggedPropertyRedirects=(ClassName="ProceduralFoliageComponent",OldPropertyName="Overlap",NewPropertyName="TileOverlap")
 
 ;Shortened animation compression algorithm class names
 +ActiveClassRedirects=(OldClassName="AnimationCompressionAlgorithm",NewClassName="AnimCompress")
diff --git a/Engine/Source/Editor/FoliageEdit/Private/ActorFactoryProceduralFoliage.cpp b/Engine/Source/Editor/FoliageEdit/Private/ActorFactoryProceduralFoliage.cpp
--- a/Engine/Source/Editor/FoliageEdit/Private/ActorFactoryProceduralFoliage.cpp
+++ b/Engine/Source/Editor/FoliageEdit/Private/ActorFactoryProceduralFoliage.cpp
@@ -6,7 +6,7 @@ ActorFactory.cpp:
 
 #include "UnrealEd.h"
 #include "ActorFactoryProceduralFoliage.h"
-#include "ProceduralFoliage.h"
+#include "ProceduralFoliageSpawner.h"
 #include "ProceduralFoliageVolume.h"
 #include "ProceduralFoliageComponent.h"
 
@@ -27,7 +27,7 @@ UActorFactoryProceduralFoliage::UActorFactoryProceduralFoliage(const FObjectInit
 
 bool UActorFactoryProceduralFoliage::CanCreateActorFrom(const FAssetData& AssetData, FText& OutErrorMsg)
 {
-	if (!AssetData.IsValid() || !AssetData.GetClass()->IsChildOf(UProceduralFoliage::StaticClass()))
+	if (!AssetData.IsValid() || !AssetData.GetClass()->IsChildOf(UProceduralFoliageSpawner::StaticClass()))
 	{
 		OutErrorMsg = NSLOCTEXT("CanCreateActor", "NoProceduralFoliage", "A valid ProceduralFoliage must be specified.");
 		return false;
@@ -39,18 +39,18 @@ bool UActorFactoryProceduralFoliage::CanCreateActorFrom(const FAssetData& AssetD
 void UActorFactoryProceduralFoliage::PostSpawnActor(UObject* Asset, AActor* NewActor)
 {
 	Super::PostSpawnActor(Asset, NewActor);
-	UProceduralFoliage* ProceduralFoliage = CastChecked<UProceduralFoliage>(Asset);
+	UProceduralFoliageSpawner* FoliageSpawner = CastChecked<UProceduralFoliageSpawner>(Asset);
 
-	UE_LOG(LogActorFactory, Log, TEXT("Actor Factory created %s"), *ProceduralFoliage->GetName());
+	UE_LOG(LogActorFactory, Log, TEXT("Actor Factory created %s"), *FoliageSpawner->GetName());
 
 	// Change properties
-	AProceduralFoliageVolume* ProceduralFoliageVolume = CastChecked<AProceduralFoliageVolume>(NewActor);
-	UProceduralFoliageComponent* ProceduralComponent = ProceduralFoliageVolume->ProceduralComponent;
+	AProceduralFoliageVolume* PFV = CastChecked<AProceduralFoliageVolume>(NewActor);
+	UProceduralFoliageComponent* ProceduralComponent = PFV->ProceduralComponent;
 	check(ProceduralComponent);
 
 	ProceduralComponent->UnregisterComponent();
 
-	ProceduralComponent->ProceduralFoliage = ProceduralFoliage;
+	ProceduralComponent->FoliageSpawner = FoliageSpawner;
 
 	// Init Component
 	ProceduralComponent->RegisterComponent();
@@ -60,21 +60,21 @@ UObject* UActorFactoryProceduralFoliage::GetAssetFromActorInstance(AActor* Insta
 {
 	check(Instance->IsA(NewActorClass));
 
-	AProceduralFoliageVolume* ProceduralFoliageVolume = CastChecked<AProceduralFoliageVolume>(Instance);
-	UProceduralFoliageComponent* ProceduralComponent = ProceduralFoliageVolume->ProceduralComponent;
+	AProceduralFoliageVolume* PFV = CastChecked<AProceduralFoliageVolume>(Instance);
+	UProceduralFoliageComponent* ProceduralComponent = PFV->ProceduralComponent;
 	check(ProceduralComponent);
 	
-	return ProceduralComponent->ProceduralFoliage;
+	return ProceduralComponent->FoliageSpawner;
 }
 
 void UActorFactoryProceduralFoliage::PostCreateBlueprint(UObject* Asset, AActor* CDO)
 {
 	if (Asset != nullptr && CDO != nullptr)
 	{
-		UProceduralFoliage* ProceduralFoliage = CastChecked<UProceduralFoliage>(Asset);
-		AProceduralFoliageVolume* PFA = CastChecked<AProceduralFoliageVolume>(CDO);
-		UProceduralFoliageComponent* ProceduralComponent = PFA->ProceduralComponent;
-		ProceduralComponent->ProceduralFoliage = ProceduralFoliage;
+		UProceduralFoliageSpawner* FoliageSpawner = CastChecked<UProceduralFoliageSpawner>(Asset);
+		AProceduralFoliageVolume* PFV = CastChecked<AProceduralFoliageVolume>(CDO);
+		UProceduralFoliageComponent* ProceduralComponent = PFV->ProceduralComponent;
+		ProceduralComponent->FoliageSpawner = FoliageSpawner;
 	}
 }
 #undef LOCTEXT_NAMESPACE
diff --git a/Engine/Source/Editor/FoliageEdit/Private/ProceduralFoliageComponentDetails.cpp b/Engine/Source/Editor/FoliageEdit/Private/ProceduralFoliageComponentDetails.cpp
--- a/Engine/Source/Editor/FoliageEdit/Private/ProceduralFoliageComponentDetails.cpp
+++ b/Engine/Source/Editor/FoliageEdit/Private/ProceduralFoliageComponentDetails.cpp
@@ -3,7 +3,7 @@
 #include "UnrealEd.h"
 #include "PropertyEditing.h"
 #include "ProceduralFoliageComponentDetails.h"
-#include "ProceduralFoliage.h"
+#include "ProceduralFoliageSpawner.h"
 #include "ProceduralFoliageComponent.h"
 #include "InstancedFoliage.h"
 #include "FoliageEdMode.h"
@@ -51,7 +51,7 @@ void FProceduralFoliageComponentDetails::CustomizeDetails( IDetailLayoutBuilder&
 	[
 		SNew(SButton)
 		.OnClicked( this, &FProceduralFoliageComponentDetails::OnResimulateClicked )
-		.ToolTipText( NSLOCTEXT("ProceduralFoliageComponentDetails","ResimulateButton_Tooltip", "Resimulates the ProceduralFoliage asset and replaces previously spawned instances" ) )
+		.ToolTipText( NSLOCTEXT("ProceduralFoliageComponentDetails","ResimulateButton_Tooltip", "Runs the procedural foliage spawner simulation. Replaces any existing instances spawned by a previous simulation." ) )
 		.IsEnabled( this, &FProceduralFoliageComponentDetails::IsResimulateEnabled )
 		[
 			SNew( STextBlock )
@@ -63,14 +63,14 @@ void FProceduralFoliageComponentDetails::CustomizeDetails( IDetailLayoutBuilder&
 
 FReply FProceduralFoliageComponentDetails::OnResimulateClicked()
 {
-	TSet<UProceduralFoliage*> UniqueProceduralFoliages;
+	TSet<UProceduralFoliageSpawner*> UniqueFoliageSpawners;
 	for( TWeakObjectPtr<UProceduralFoliageComponent>& Component : SelectedComponents )
 	{
-		if( Component.IsValid() && Component->ProceduralFoliage )
+		if( Component.IsValid() && Component->FoliageSpawner )
 		{
-			if( !UniqueProceduralFoliages.Contains( Component->ProceduralFoliage ) )
+			if( !UniqueFoliageSpawners.Contains( Component->FoliageSpawner ) )
 			{
-				UniqueProceduralFoliages.Add( Component->ProceduralFoliage );
+				UniqueFoliageSpawners.Add(Component->FoliageSpawner);
 			}
 
 			TArray <FDesiredFoliageInstance> DesiredFoliageInstances;
@@ -87,7 +87,7 @@ bool FProceduralFoliageComponentDetails::IsResimulateEnabled() const
 {
 	for(const TWeakObjectPtr<UProceduralFoliageComponent>& Component : SelectedComponents)
 	{
-		if(Component.IsValid() && Component->ProceduralFoliage)
+		if(Component.IsValid() && Component->FoliageSpawner)
 		{
 			return true;
 		}
diff --git a/Engine/Source/Editor/FoliageEdit/Private/ProceduralFoliageComponentVisualizer.cpp b/Engine/Source/Editor/FoliageEdit/Private/ProceduralFoliageComponentVisualizer.cpp
--- a/Engine/Source/Editor/FoliageEdit/Private/ProceduralFoliageComponentVisualizer.cpp
+++ b/Engine/Source/Editor/FoliageEdit/Private/ProceduralFoliageComponentVisualizer.cpp
@@ -2,7 +2,7 @@
 
 #include "UnrealEd.h"
 #include "ProceduralFoliageComponentVisualizer.h"
-#include "ProceduralFoliage.h"
+#include "ProceduralFoliageSpawner.h"
 #include "ProceduralFoliageComponent.h"
 
 
@@ -13,12 +13,12 @@ void FProceduralFoliageComponentVisualizer::DrawVisualization(const UActorCompon
 {
 	if (const UProceduralFoliageComponent* ProcComponent = Cast<const UProceduralFoliageComponent>(Component))
 	{
-		if (ProcComponent->ProceduralFoliage && ProcComponent->bHideDebugTiles == false)
+		if (ProcComponent->FoliageSpawner && ProcComponent->bShowDebugTiles == true)
 		{
-			const float TileSize = ProcComponent->ProceduralFoliage->TileSize;
+			const float TileSize = ProcComponent->FoliageSpawner->TileSize;
 			const FVector TileSizeOffset(TileSize, TileSize, 0.f);
 			const FVector TileLocation = ProcComponent->GetWorldPosition();
-			const FVector OverlapV(ProcComponent->Overlap, ProcComponent->Overlap, 0.f);
+			const FVector OverlapV(ProcComponent->TileOverlap, ProcComponent->TileOverlap, 0.f);
 			
 			int32 XOffset, YOffset;
 			int32 NumX = 0;
diff --git a/Engine/Source/Editor/FoliageEdit/Private/ProceduralFoliageFactory.cpp b/Engine/Source/Editor/FoliageEdit/Private/ProceduralFoliageFactory.cpp
deleted file mode 100644
diff --git a/Engine/Source/Editor/FoliageEdit/Private/ProceduralFoliageFactory.h b/Engine/Source/Editor/FoliageEdit/Private/ProceduralFoliageFactory.h
deleted file mode 100644
diff --git a/Engine/Source/Editor/FoliageEdit/Private/ProceduralFoliageSpawnerFactory.cpp b/Engine/Source/Editor/FoliageEdit/Private/ProceduralFoliageSpawnerFactory.cpp
new file mode 100644
--- /dev/null
+++ b/Engine/Source/Editor/FoliageEdit/Private/ProceduralFoliageSpawnerFactory.cpp
@@ -0,0 +1,26 @@
+// Copyright 1998-2015 Epic Games, Inc. All Rights Reserved.
+
+#include "UnrealEd.h"
+#include "ProceduralFoliageSpawnerFactory.h"
+#include "IAssetTypeActions.h"
+#include "ProceduralFoliageSpawner.h"
+
+UProceduralFoliageSpawnerFactory::UProceduralFoliageSpawnerFactory(const FObjectInitializer& ObjectInitializer)
+	: Super(ObjectInitializer)
+{
+	bCreateNew = true;
+	bEditAfterNew = true;
+	SupportedClass = UProceduralFoliageSpawner::StaticClass();
+}
+
+UObject* UProceduralFoliageSpawnerFactory::FactoryCreateNew(UClass* Class, UObject* InParent, FName Name, EObjectFlags Flags, UObject* Context, FFeedbackContext* Warn)
+{
+	auto NewProceduralFoliage = NewObject<UProceduralFoliageSpawner>(InParent, Class, Name, Flags | RF_Transactional);
+
+	return NewProceduralFoliage;
+}
+
+uint32 UProceduralFoliageSpawnerFactory::GetMenuCategories() const
+{
+	return EAssetTypeCategories::Misc;
+}
\ No newline at end of file
diff --git a/Engine/Source/Editor/FoliageEdit/Private/ProceduralFoliageSpawnerFactory.h b/Engine/Source/Editor/FoliageEdit/Private/ProceduralFoliageSpawnerFactory.h
new file mode 100644
--- /dev/null
+++ b/Engine/Source/Editor/FoliageEdit/Private/ProceduralFoliageSpawnerFactory.h
@@ -0,0 +1,21 @@
+// Copyright 1998-2015 Epic Games, Inc. All Rights Reserved.
+
+/**
+ * Factory for ProceduralFoliage assets
+ */
+
+#pragma once
+
+#include "ProceduralFoliageSpawnerFactory.generated.h"
+
+
+UCLASS()
+class UProceduralFoliageSpawnerFactory : public UFactory
+{
+	GENERATED_UCLASS_BODY()
+
+	// UFactory interface
+	virtual UObject* FactoryCreateNew(UClass* Class, UObject* InParent, FName Name, EObjectFlags Flags, UObject* Context, FFeedbackContext* Warn) override;
+	virtual uint32 GetMenuCategories() const;
+	// End of UFactory interface
+};
diff --git a/Engine/Source/Runtime/Foliage/Private/ProceduralFoliage.cpp b/Engine/Source/Runtime/Foliage/Private/ProceduralFoliage.cpp
deleted file mode 100644
diff --git a/Engine/Source/Runtime/Foliage/Private/ProceduralFoliageComponent.cpp b/Engine/Source/Runtime/Foliage/Private/ProceduralFoliageComponent.cpp
--- a/Engine/Source/Runtime/Foliage/Private/ProceduralFoliageComponent.cpp
+++ b/Engine/Source/Runtime/Foliage/Private/ProceduralFoliageComponent.cpp
@@ -5,7 +5,7 @@
 #include "ProceduralFoliageTile.h"
 #include "InstancedFoliage.h"
 #include "InstancedFoliageActor.h"
-#include "ProceduralFoliage.h"
+#include "ProceduralFoliageSpawner.h"
 
 #include "Async/Async.h"
 
@@ -14,10 +14,10 @@
 UProceduralFoliageComponent::UProceduralFoliageComponent(const FObjectInitializer& ObjectInitializer)
 : Super(ObjectInitializer)
 {
-	Overlap = 0.f;
+	TileOverlap = 0.f;
 	ProceduralGuid = FGuid::NewGuid();
 #if WITH_EDITORONLY_DATA
-	bHideDebugTiles = true;
+	bShowDebugTiles = false;
 #endif
 }
 
@@ -51,15 +51,15 @@ void UProceduralFoliageComponent::GetTilesLayout(int32& MinXIdx, int32& MinYIdx,
 {
 	if(UBrushComponent* Brush = SpawningVolume->GetBrushComponent())
 	{
-		const FVector MinPosition = Brush->Bounds.GetBox().Min + Overlap;
-		const FVector MaxPosition = Brush->Bounds.GetBox().Max - Overlap;
+		const FVector MinPosition = Brush->Bounds.GetBox().Min + TileOverlap;
+		const FVector MaxPosition = Brush->Bounds.GetBox().Max - TileOverlap;
 
 		//we want to find the bottom left tile that contains this MinPosition
-		MinXIdx = FMath::FloorToInt(MinPosition.X / ProceduralFoliage->TileSize);
-		MinYIdx = FMath::FloorToInt(MinPosition.Y / ProceduralFoliage->TileSize);
+		MinXIdx = FMath::FloorToInt(MinPosition.X / FoliageSpawner->TileSize);
+		MinYIdx = FMath::FloorToInt(MinPosition.Y / FoliageSpawner->TileSize);
 
-		const int32 MaxXIdx = FMath::FloorToInt(MaxPosition.X / ProceduralFoliage->TileSize);
-		const int32 MaxYIdx = FMath::FloorToInt(MaxPosition.Y / ProceduralFoliage->TileSize);
+		const int32 MaxXIdx = FMath::FloorToInt(MaxPosition.X / FoliageSpawner->TileSize);
+		const int32 MaxYIdx = FMath::FloorToInt(MaxPosition.Y / FoliageSpawner->TileSize);
 
 		NumX = (MaxXIdx - MinXIdx) + 1;
 		NumY = (MaxYIdx - MinYIdx) + 1;
@@ -71,7 +71,7 @@ void UProceduralFoliageComponent::GetTilesLayout(int32& MinXIdx, int32& MinYIdx,
 FVector UProceduralFoliageComponent::GetWorldPosition() const
 {
 	UBrushComponent* Brush = SpawningVolume->GetBrushComponent();
-	if (Brush == nullptr || ProceduralFoliage == nullptr)
+	if (Brush == nullptr || FoliageSpawner == nullptr)
 	{
 		return FVector::ZeroVector;
 	}
@@ -82,27 +82,27 @@ FVector UProceduralFoliageComponent::GetWorldPosition() const
 	float HalfHeight;
 	GetTilesLayout(X1, Y1, NumX, NumY, HalfHeight);
 
-	return FVector(X1 * ProceduralFoliage->TileSize, Y1 * ProceduralFoliage->TileSize, Brush->Bounds.Origin.Z);
+	return FVector(X1 * FoliageSpawner->TileSize, Y1 * FoliageSpawner->TileSize, Brush->Bounds.Origin.Z);
 }
 
 bool UProceduralFoliageComponent::SpawnTiles(TArray<FDesiredFoliageInstance>& OutInstances)
 {
 #if WITH_EDITOR
-	if (ProceduralFoliage && SpawningVolume && SpawningVolume->GetBrushComponent())
+	if (FoliageSpawner && SpawningVolume && SpawningVolume->GetBrushComponent())
 	{
 		/** Constants for laying out the overlapping tile grid*/
-		const float InnerTileSize = ProceduralFoliage->TileSize;
+		const float InnerTileSize = FoliageSpawner->TileSize;
 		const FVector2D InnerTileV(InnerTileSize, InnerTileSize);
 		const FVector2D InnerTileX(InnerTileSize, 0.f);
 		const FVector2D InnerTileY(0.f, InnerTileSize);
 
-		const float OuterTileSize = ProceduralFoliage->TileSize + Overlap;
+		const float OuterTileSize = FoliageSpawner->TileSize + TileOverlap;
 		const FVector2D OuterTileV(OuterTileSize, OuterTileSize);
 		const FVector2D OuterTileX(OuterTileSize, 0.f);
 		const FVector2D OuterTileY(0.f, OuterTileSize);
 
-		const FVector2D OverlapX(Overlap, 0.f);
-		const FVector2D OverlapY(0.f, Overlap);
+		const FVector2D OverlapX(TileOverlap, 0.f);
+		const FVector2D OverlapY(0.f, TileOverlap);
 
 		TArray<TFuture< TArray<FDesiredFoliageInstance>* >> Futures;
 		
@@ -117,24 +117,24 @@ bool UProceduralFoliageComponent::SpawnTiles(TArray<FDesiredFoliageInstance>& Ou
 		FThreadSafeCounter* LastCancelPtr = &LastCancel;
 
 		const FVector WorldPosition = GetWorldPosition();
-		ProceduralFoliage->SimulateIfNeeded();
+		FoliageSpawner->SimulateIfNeeded();
 
 		for (int32 X = 0; X < NumTilesX; ++X)
 		{
 			for (int32 Y = 0; Y < NumTilesY; ++Y)
 			{
 				//We have to get the tiles and create new one to build on main thread
-				const UProceduralFoliageTile* Tile = ProceduralFoliage->GetRandomTile(X + XOffset, Y + YOffset);
-				if (Tile == nullptr)	//simulation was cancelled or failed
+				const UProceduralFoliageTile* Tile = FoliageSpawner->GetRandomTile(X + XOffset, Y + YOffset);
+				if (Tile == nullptr)	//simulation was canceled or failed
 				{
 					return false;
 				}
 
-				const UProceduralFoliageTile* RightTile = (X + 1 < NumTilesX) ? ProceduralFoliage->GetRandomTile(X + XOffset + 1, Y + YOffset) : nullptr;
-				const UProceduralFoliageTile* TopTile = (Y + 1 < NumTilesY) ? ProceduralFoliage->GetRandomTile(X + XOffset, Y + YOffset+ 1) : nullptr;
-				const UProceduralFoliageTile* TopRightTile = (RightTile && TopTile) ? ProceduralFoliage->GetRandomTile(X + XOffset + 1, Y + YOffset + 1) : nullptr;
+				const UProceduralFoliageTile* RightTile = ( X + 1 < NumTilesX ) ? FoliageSpawner->GetRandomTile(X + XOffset + 1, Y + YOffset) : nullptr;
+				const UProceduralFoliageTile* TopTile = ( Y + 1 < NumTilesY ) ? FoliageSpawner->GetRandomTile(X + XOffset, Y + YOffset + 1) : nullptr;
+				const UProceduralFoliageTile* TopRightTile = ( RightTile && TopTile ) ? FoliageSpawner->GetRandomTile(X + XOffset + 1, Y + YOffset + 1) : nullptr;
 
-				UProceduralFoliageTile* JTile = ProceduralFoliage->CreateTempTile();
+				UProceduralFoliageTile* JTile = FoliageSpawner->CreateTempTile();
 
 				Futures.Add(Async<TArray<FDesiredFoliageInstance>*>(EAsyncExecution::ThreadPool, [=]()
 				{
@@ -147,32 +147,32 @@ bool UProceduralFoliageComponent::SpawnTiles(TArray<FDesiredFoliageInstance>& Ou
 					const FTransform TileTM(OrientedOffset + WorldPosition);
 
 					//copy the inner tile
-					const FBox2D InnerBox = GetTileRegion(X, Y, NumTilesX, NumTilesY, InnerTileSize, Overlap);
-					CopyTileInstances(Tile, JTile, InnerBox, FTransform::Identity, Overlap);
+					const FBox2D InnerBox = GetTileRegion(X, Y, NumTilesX, NumTilesY, InnerTileSize, TileOverlap);
+					CopyTileInstances(Tile, JTile, InnerBox, FTransform::Identity, TileOverlap);
 
 
 					if (RightTile)
 					{
 						//Add overlap to the right
-						const FBox2D RightBox(FVector2D(-Overlap, InnerBox.Min.Y), FVector2D(Overlap, InnerBox.Max.Y));
+						const FBox2D RightBox(FVector2D(-TileOverlap, InnerBox.Min.Y), FVector2D(TileOverlap, InnerBox.Max.Y));
 						const FTransform RightTM(FVector(InnerTileSize, 0.f, 0.f));
-						CopyTileInstances(RightTile, JTile, RightBox, RightTM, Overlap);
+						CopyTileInstances(RightTile, JTile, RightBox, RightTM, TileOverlap);
 					}
 
 					if (TopTile)
 					{
 						//Add overlap to the top
-						const FBox2D TopBox(FVector2D(InnerBox.Min.X, -Overlap), FVector2D(InnerBox.Max.X, Overlap));
+						const FBox2D TopBox(FVector2D(InnerBox.Min.X, -TileOverlap), FVector2D(InnerBox.Max.X, TileOverlap));
 						const FTransform TopTM(FVector(0.f, InnerTileSize, 0.f));
-						CopyTileInstances(TopTile, JTile, TopBox, TopTM, Overlap);
+						CopyTileInstances(TopTile, JTile, TopBox, TopTM, TileOverlap);
 					}
 
 					if (TopRightTile)
 					{
 						//Add overlap to the top right
-						const FBox2D TopRightBox(FVector2D(-Overlap, -Overlap), FVector2D(Overlap, Overlap));
+						const FBox2D TopRightBox(FVector2D(-TileOverlap, -TileOverlap), FVector2D(TileOverlap, TileOverlap));
 						const FTransform TopRightTM(FVector(InnerTileSize, InnerTileSize, 0.f));
-						CopyTileInstances(TopRightTile, JTile, TopRightBox, TopRightTM, Overlap);
+						CopyTileInstances(TopRightTile, JTile, TopRightBox, TopRightTM, TileOverlap);
 					}
 
 					TArray<FDesiredFoliageInstance>* DesiredInstances = new TArray<FDesiredFoliageInstance>();
diff --git a/Engine/Source/Runtime/Foliage/Private/ProceduralFoliageSpawner.cpp b/Engine/Source/Runtime/Foliage/Private/ProceduralFoliageSpawner.cpp
new file mode 100644
--- /dev/null
+++ b/Engine/Source/Runtime/Foliage/Private/ProceduralFoliageSpawner.cpp
@@ -0,0 +1,208 @@
+// Copyright 1998-2015 Epic Games, Inc. All Rights Reserved.
+
+#include "FoliagePrivate.h"
+#include "ProceduralFoliageSpawner.h"
+#include "ProceduralFoliageTile.h"
+#include "Components/BoxComponent.h"
+#include "PhysicsPublic.h"
+#include "Async/Async.h"
+
+#define LOCTEXT_NAMESPACE "ProceduralFoliage"
+
+UProceduralFoliageSpawner::UProceduralFoliageSpawner(const FObjectInitializer& ObjectInitializer)
+	: Super(ObjectInitializer)
+{
+	TileSize = 10000;	//100 m
+	NumUniqueTiles = 10;
+	RandomSeed = 42;
+}
+
+#if WITH_EDITOR
+void UProceduralFoliageSpawner::PostEditChangeProperty(FPropertyChangedEvent& PropertyChangedEvent)
+{
+	bNeedsSimulation = true;
+}
+#endif
+
+UProceduralFoliageTile* UProceduralFoliageSpawner::CreateTempTile()
+{
+	UProceduralFoliageTile* TmpTile = NewObject<UProceduralFoliageTile>(this);
+	TmpTile->InitSimulation(this, 0);
+
+	return TmpTile;
+}
+
+void UProceduralFoliageSpawner::CreateProceduralFoliageInstances()
+{
+	for(FProceduralFoliageTypeData& TypeData : Types)
+	{
+		if( TypeData.Type )
+		{
+			TypeData.TypeInstance = NewObject<UFoliageType_InstancedStaticMesh>(this, TypeData.Type);
+			TypeData.ChangeCount = TypeData.TypeInstance->ChangeCount;
+		}
+		else
+		{
+			TypeData.TypeInstance = nullptr;
+		}
+	}
+}
+
+void UProceduralFoliageSpawner::SetClean()
+{
+	for (FProceduralFoliageTypeData& TypeData : Types)
+	{
+		if( TypeData.Type && TypeData.TypeInstance )
+		{
+			TypeData.TypeInstance->ChangeCount = TypeData.Type->GetDefaultObject<UFoliageType_InstancedStaticMesh>()->ChangeCount;
+		}
+	}
+
+	bNeedsSimulation = false;
+}
+
+// Custom serialization version for all packages containing ProceduralFoliage
+struct FProceduralFoliageCustomVersion
+{
+	enum Type
+	{
+		// Before any version changes were made in the plugin
+		BeforeCustomVersionWasAdded = 0,
+		// We serialize the type clean change map to ensure validity
+		SerializeTypeMap = 1,
+		// -----<new versions can be added above this line>-------------------------------------------------
+		VersionPlusOne,
+		LatestVersion = VersionPlusOne - 1
+	};
+
+	// The GUID for this custom version number
+	const static FGuid GUID;
+
+private:
+	FProceduralFoliageCustomVersion() {}
+};
+
+const FGuid FProceduralFoliageCustomVersion::GUID(0xaafe32bd, 0x53954c14, 0xb66a5e25, 0x1032d1dd);
+// Register the custom version with core
+FCustomVersionRegistration GRegisterProceduralFoliageCustomVersion(FProceduralFoliageCustomVersion::GUID, FProceduralFoliageCustomVersion::LatestVersion, TEXT("ProceduralFoliageVer"));
+
+void UProceduralFoliageSpawner::Serialize(FArchive& Ar)
+{
+	Super::Serialize(Ar);
+
+	Ar.UsingCustomVersion(FProceduralFoliageCustomVersion::GUID);
+}
+
+bool UProceduralFoliageSpawner::AnyDirty() const
+{
+	bool bDirty = bNeedsSimulation;
+
+	for (const FProceduralFoliageTypeData& TypeData : Types)
+	{
+		if (TypeData.Type)
+		{
+			if (TypeData.TypeInstance && TypeData.Type->GetDefaultObject<UFoliageType_InstancedStaticMesh>()->ChangeCount != TypeData.ChangeCount)
+			{
+				bDirty = true;
+				break;
+			}
+			else if (!TypeData.TypeInstance)
+			{
+				bDirty = true;
+				break;
+			}
+		}
+	}
+
+	return bDirty;
+}
+
+void UProceduralFoliageSpawner::SimulateIfNeeded()
+{
+	//if (AnyDirty())	@todo: for now we must force simulation every time since precomputed tiles are weak pointers
+	{
+		Simulate();
+	}
+
+}
+
+const UProceduralFoliageTile* UProceduralFoliageSpawner::GetRandomTile(int32 X, int32 Y)
+{
+	if (PrecomputedTiles.Num())
+	{
+		FRandomStream HashStream;	//using this as a hash function
+		HashStream.Initialize(X);
+		const float XRand = HashStream.FRand();
+		HashStream.Initialize(Y);
+		const float YRand = HashStream.FRand();
+		const int32 RandomNumber = (RAND_MAX * XRand / (YRand + 0.01f));
+		const int32 Idx = FMath::Clamp(RandomNumber % PrecomputedTiles.Num(), 0, PrecomputedTiles.Num() - 1);
+		return PrecomputedTiles[Idx].Get();
+	}
+
+	return nullptr;
+}
+
+void UProceduralFoliageSpawner::Simulate(int32 NumSteps)
+{
+	RandomStream.Initialize(RandomSeed);
+	CreateProceduralFoliageInstances();
+
+	LastCancel.Increment();
+
+	PrecomputedTiles.Empty();
+	TArray<TFuture< UProceduralFoliageTile* >> Futures;
+
+	for (int i = 0; i < NumUniqueTiles; ++i)
+	{
+		UProceduralFoliageTile* NewTile = NewObject<UProceduralFoliageTile>(this);
+		const int32 RandomNumber = GetRandomNumber();
+		const int32 LastCancelInit = LastCancel.GetValue();
+
+		Futures.Add(Async<UProceduralFoliageTile*>(EAsyncExecution::ThreadPool, [=]()
+		{
+			NewTile->Simulate(this, RandomNumber, NumSteps, LastCancelInit);
+			return NewTile;
+		}));
+	}
+
+	const FText StatusMessage = LOCTEXT("SimulateProceduralFoliage", "Simulate ProceduralFoliage...");
+	GWarn->BeginSlowTask(StatusMessage, true, true);
+	
+	const int32 TotalTasks = Futures.Num();
+
+	bool bCancelled = false;
+
+	for (int32 FutureIdx = 0; FutureIdx < Futures.Num(); ++FutureIdx)
+	{
+		while (Futures[FutureIdx].WaitFor(FTimespan(0, 0, 0, 0, 100)) == false)		//sleep for 100ms if not ready. Needed so cancel is responsive
+		{
+			GWarn->StatusUpdate(FutureIdx, TotalTasks, LOCTEXT("SimulateProceduralFoliage", "Simulate ProceduralFoliage..."));
+			if (GWarn->ReceivedUserCancel() && bCancelled == false)	//If the user wants to cancel just increment this. Tiles compare against the original count and cancel if needed
+			{
+				LastCancel.Increment();
+				bCancelled = true;
+			}
+		}
+
+		PrecomputedTiles.Add(Futures[FutureIdx].Get());		//Even if we cancel we block until threads have exited safely to ensure memory isn't GCed
+	}
+
+	GWarn->EndSlowTask();
+
+	if (bCancelled)
+	{
+		PrecomputedTiles.Empty();
+	}
+	else
+	{
+		SetClean();
+	}
+}
+
+int32 UProceduralFoliageSpawner::GetRandomNumber()
+{
+	return RandomStream.FRand() * RAND_MAX;
+}
+
+#undef LOCTEXT_NAMESPACE
\ No newline at end of file
diff --git a/Engine/Source/Runtime/Foliage/Private/ProceduralFoliageTile.cpp b/Engine/Source/Runtime/Foliage/Private/ProceduralFoliageTile.cpp
--- a/Engine/Source/Runtime/Foliage/Private/ProceduralFoliageTile.cpp
+++ b/Engine/Source/Runtime/Foliage/Private/ProceduralFoliageTile.cpp
@@ -2,7 +2,7 @@
 
 #include "FoliagePrivate.h"
 #include "ProceduralFoliageTile.h"
-#include "ProceduralFoliage.h"
+#include "ProceduralFoliageSpawner.h"
 #include "ProceduralFoliageBroadphase.h"
 #include "InstancedFoliageActor.h"
 
@@ -192,7 +192,7 @@ void UProceduralFoliageTile::SpreadSeeds(TArray<FProceduralFoliageInstance*>& Ne
 
 void UProceduralFoliageTile::AddRandomSeeds(TArray<FProceduralFoliageInstance*>& OutInstances)
 {
-	const float SizeTenM2 = (ProceduralFoliage->TileSize * ProceduralFoliage->TileSize) / (1000.f * 1000.f);
+	const float SizeTenM2 = ( FoliageSpawner->TileSize * FoliageSpawner->TileSize ) / ( 1000.f * 1000.f );
 
 	TMap<int32,float> MaxShadeRadii;
 	TMap<int32, float> MaxCollisionRadii;
@@ -201,7 +201,7 @@ void UProceduralFoliageTile::AddRandomSeeds(TArray<FProceduralFoliageInstance*>&
 
 	TArray<const UFoliageType_InstancedStaticMesh*> TypesToSeed;
 
-	for (const FProceduralFoliageTypeData& Data : ProceduralFoliage->GetTypes())
+	for (const FProceduralFoliageTypeData& Data : FoliageSpawner->GetTypes())
 	{
 		if (UserCancelled()){ return; }
 		const UFoliageType_InstancedStaticMesh* Type = Data.TypeInstance;
@@ -217,7 +217,7 @@ void UProceduralFoliageTile::AddRandomSeeds(TArray<FProceduralFoliageInstance*>&
 			}
 
 			{	//save the random stream per type
-				RandomStreamPerType.Add(Type, FRandomStream(Type->DistributionSeed + ProceduralFoliage->RandomSeed));
+				RandomStreamPerType.Add(Type, FRandomStream(Type->DistributionSeed + FoliageSpawner->RandomSeed));
 			}
 
 			{	//compute the needed offsets for initial seed variance
@@ -282,8 +282,8 @@ void UProceduralFoliageTile::AddRandomSeeds(TArray<FProceduralFoliageInstance*>&
 			}
 			else
 			{
-				InitX = TypeRandomStream.FRandRange(0, ProceduralFoliage->TileSize);
-				InitY = TypeRandomStream.FRandRange(0, ProceduralFoliage->TileSize);
+				InitX = TypeRandomStream.FRandRange(0, FoliageSpawner->TileSize);
+				InitY = TypeRandomStream.FRandRange(0, FoliageSpawner->TileSize);
 				NeededRadius = MaxShadeRadii.FindRef(Type->DistributionSeed);
 			}
 
@@ -364,17 +364,17 @@ void UProceduralFoliageTile::FlushPendingRemovals()
 	PendingRemovals.Empty();
 }
 
-void UProceduralFoliageTile::InitSimulation(const UProceduralFoliage* InProceduralFoliage, const int32 RandomSeed)
+void UProceduralFoliageTile::InitSimulation(const UProceduralFoliageSpawner* InFoliageSpawner, const int32 RandomSeed)
 {
 	RandomStream.Initialize(RandomSeed);
-	ProceduralFoliage = InProceduralFoliage;
+	FoliageSpawner = InFoliageSpawner;
 	SimulationStep = 0;
-	Broadphase = FProceduralFoliageBroadphase(ProceduralFoliage->TileSize);
+	Broadphase = FProceduralFoliageBroadphase(FoliageSpawner->TileSize);
 }
 
 bool UProceduralFoliageTile::UserCancelled() const
 {
-	return ProceduralFoliage->LastCancel.GetValue() != LastCancel;
+	return FoliageSpawner->LastCancel.GetValue() != LastCancel;
 }
 
 
@@ -404,7 +404,7 @@ void UProceduralFoliageTile::StartSimulation(const int32 MaxNumSteps, bool bShad
 {
 	int32 MaxSteps = 0;
 
-	for (const FProceduralFoliageTypeData& Data : ProceduralFoliage->GetTypes())
+	for (const FProceduralFoliageTypeData& Data : FoliageSpawner->GetTypes())
 	{
 		const UFoliageType_InstancedStaticMesh* TypeInstance = Data.TypeInstance;
 		if (TypeInstance && TypeInstance->bGrowsInShade == bShadeGrowth)
@@ -429,10 +429,10 @@ void UProceduralFoliageTile::StartSimulation(const int32 MaxNumSteps, bool bShad
 	InstancesToArray();
 }
 
-void UProceduralFoliageTile::Simulate(const UProceduralFoliage* InProceduralFoliage, const int32 RandomSeed, const int32 MaxNumSteps, const int32 InLastCancel)
+void UProceduralFoliageTile::Simulate(const UProceduralFoliageSpawner* InFoliageSpawner, const int32 RandomSeed, const int32 MaxNumSteps, const int32 InLastCancel)
 {
 	LastCancel = InLastCancel;
-	InitSimulation(InProceduralFoliage, RandomSeed);
+	InitSimulation(InFoliageSpawner, RandomSeed);
 
 	StartSimulation(MaxNumSteps, false);
 	StartSimulation(MaxNumSteps, true);
diff --git a/Engine/Source/Runtime/Foliage/Private/ProceduralFoliageVolume.cpp b/Engine/Source/Runtime/Foliage/Private/ProceduralFoliageVolume.cpp
--- a/Engine/Source/Runtime/Foliage/Private/ProceduralFoliageVolume.cpp
+++ b/Engine/Source/Runtime/Foliage/Private/ProceduralFoliageVolume.cpp
@@ -3,7 +3,7 @@
 #include "FoliagePrivate.h"
 #include "ProceduralFoliageVolume.h"
 #include "ProceduralFoliageComponent.h"
-#include "ProceduralFoliage.h"
+#include "ProceduralFoliageSpawner.h"
 
 AProceduralFoliageVolume::AProceduralFoliageVolume(const FObjectInitializer& ObjectInitializer)
 : Super(ObjectInitializer)
@@ -26,9 +26,9 @@ bool AProceduralFoliageVolume::GetReferencedContentObjects(TArray<UObject*>& Obj
 {
 	Super::GetReferencedContentObjects(Objects);
 
-	if (ProceduralComponent && ProceduralComponent->ProceduralFoliage)
+	if (ProceduralComponent && ProceduralComponent->FoliageSpawner)
 	{
-		Objects.Add(ProceduralComponent->ProceduralFoliage);
+		Objects.Add(ProceduralComponent->FoliageSpawner);
 	}
 	return true;
 }
diff --git a/Engine/Source/Runtime/Foliage/Public/ProceduralFoliage.h b/Engine/Source/Runtime/Foliage/Public/ProceduralFoliage.h
deleted file mode 100644
diff --git a/Engine/Source/Runtime/Foliage/Public/ProceduralFoliageComponent.h b/Engine/Source/Runtime/Foliage/Public/ProceduralFoliageComponent.h
--- a/Engine/Source/Runtime/Foliage/Public/ProceduralFoliageComponent.h
+++ b/Engine/Source/Runtime/Foliage/Public/ProceduralFoliageComponent.h
@@ -5,7 +5,7 @@
 #include "ProceduralFoliageInstance.h"
 #include "ProceduralFoliageComponent.generated.h"
 
-class UProceduralFoliage;
+class UProceduralFoliageSpawner;
 class AProceduralFoliageLevelInfo;
 struct FDesiredFoliageInstance;
 
@@ -14,17 +14,18 @@ class FOLIAGE_API UProceduralFoliageComponent : public UActorComponent
 {
 	GENERATED_UCLASS_BODY()
 
-	/** The overlap in Cms between two tiles. */
+	/** The procedural foliage spawner used to generate foliage instances within this volume. */
 	UPROPERTY(Category = "ProceduralFoliage", BlueprintReadWrite, EditAnywhere)
-	float Overlap;
+	UProceduralFoliageSpawner* FoliageSpawner;
 
-	UPROPERTY(Category = "ProceduralFoliage", BlueprintReadWrite, EditAnywhere, meta=(DisplayName="ProceduralFoliage Asset") )
-	UProceduralFoliage* ProceduralFoliage;
+	/** The amount of overlap between simulation tiles (in cm). */
+	UPROPERTY(Category = "ProceduralFoliage", BlueprintReadWrite, EditAnywhere)
+	float TileOverlap;
 
 #if WITH_EDITORONLY_DATA
-	/** The overlap in Cms between two tiles. */
+	/** Whether to visualize the tiles used for the foliage spawner simulation */
 	UPROPERTY(Category = "ProceduralFoliage", BlueprintReadWrite, EditAnywhere)
-	bool bHideDebugTiles;
+	bool bShowDebugTiles;
 #endif
 
 	bool SpawnProceduralContent(TArray<FDesiredFoliageInstance>& OutFoliageInstances);
diff --git a/Engine/Source/Runtime/Foliage/Public/ProceduralFoliageSpawner.h b/Engine/Source/Runtime/Foliage/Public/ProceduralFoliageSpawner.h
new file mode 100644
--- /dev/null
+++ b/Engine/Source/Runtime/Foliage/Public/ProceduralFoliageSpawner.h
@@ -0,0 +1,88 @@
+// Copyright 1998-2015 Epic Games, Inc. All Rights Reserved.
+
+#pragma once
+
+#include "FoliageType_InstancedStaticMesh.h"
+#include "ProceduralFoliageInstance.h"
+#include "ProceduralFoliageSpawner.generated.h"
+
+class UProceduralFoliageTile;
+
+
+USTRUCT()
+struct FProceduralFoliageTypeData
+{
+	GENERATED_USTRUCT_BODY()
+
+	UPROPERTY(Category = ProceduralFoliageSimulation, EditAnywhere, BlueprintReadOnly)
+	TSubclassOf<UFoliageType_InstancedStaticMesh> Type;
+
+	UPROPERTY(transient)
+	UFoliageType_InstancedStaticMesh* TypeInstance;
+
+	UPROPERTY()
+	int32 ChangeCount;
+};
+
+UCLASS(BlueprintType, Blueprintable)
+class FOLIAGE_API UProceduralFoliageSpawner : public UObject
+{
+	GENERATED_UCLASS_BODY()
+
+	/**The seed used for generating the randomness of the simulation. */
+	UPROPERTY(Category = ProceduralFoliageSimulation, EditAnywhere, BlueprintReadOnly)
+	int32 RandomSeed;
+
+	/**Cm size of the tile along one axis. The tile is square so the total area is TileSize*TileSize */
+	UPROPERTY(Category = ProceduralFoliageSimulation, EditAnywhere, BlueprintReadOnly)
+	float TileSize;
+
+	UPROPERTY(Category = ProceduralFoliageSimulation, EditAnywhere, BlueprintReadOnly)
+	int32 NumUniqueTiles;
+
+	FThreadSafeCounter LastCancel;
+
+private:
+	UPROPERTY(Category = ProceduralFoliageSimulation, EditAnywhere)
+	TArray<FProceduralFoliageTypeData> Types;
+
+	UPROPERTY()
+	bool bNeedsSimulation;
+
+public:
+	UFUNCTION(BlueprintCallable, Category = ProceduralFoliageSimulation)
+	void Simulate(int32 NumSteps = -1);
+
+	int32 GetRandomNumber();
+
+	const TArray<FProceduralFoliageTypeData>& GetTypes() const { return Types; }
+
+	/** Returns the instances that need to spawn for a given min,max region */
+	void GetInstancesToSpawn(TArray<FProceduralFoliageInstance>& OutInstances, const FVector& Min = FVector(-FLT_MAX, -FLT_MAX, -FLT_MAX), const FVector& Max = FVector(FLT_MAX, FLT_MAX, FLT_MAX) ) const;
+
+#if WITH_EDITOR
+	virtual void PostEditChangeProperty(FPropertyChangedEvent& PropertyChangedEvent);
+#endif
+
+	virtual void Serialize(FArchive& Ar);
+
+	/** Simulates tiles if needed */
+	void SimulateIfNeeded();
+
+	/** Takes a tile index and returns a random tile associated with that index. */
+	const UProceduralFoliageTile* GetRandomTile(int32 X, int32 Y);
+
+	/** Creates a temporary empty tile with the appropriate settings created for it. */
+	UProceduralFoliageTile* CreateTempTile();
+
+private:
+	void CreateProceduralFoliageInstances();
+
+	bool AnyDirty() const;
+
+	void SetClean();
+private:
+	TArray<TWeakObjectPtr<UProceduralFoliageTile>> PrecomputedTiles;
+
+	FRandomStream RandomStream;
+};
diff --git a/Engine/Source/Runtime/Foliage/Public/ProceduralFoliageTile.h b/Engine/Source/Runtime/Foliage/Public/ProceduralFoliageTile.h
--- a/Engine/Source/Runtime/Foliage/Public/ProceduralFoliageTile.h
+++ b/Engine/Source/Runtime/Foliage/Public/ProceduralFoliageTile.h
@@ -9,8 +9,7 @@
 
 struct FProceduralFoliageInstance;
 struct FProceduralFoliageOverlap;
-class UProceduralFoliageComponent;
-class UProceduralFoliage;
+class UProceduralFoliageSpawner;
 class UFoliageType_InstancedStaticMesh;
 
 UCLASS()
@@ -18,10 +17,7 @@ class FOLIAGE_API UProceduralFoliageTile : public UObject
 {
 	GENERATED_UCLASS_BODY()
 
-	UPROPERTY()
-	const UProceduralFoliage* ProceduralFoliage;
-
-	void Simulate(const UProceduralFoliage* ProceduralFoliage, const int32 RandomSeed, const int32 MaxNumSteps, const int32 InLastCancel);
+	void Simulate(const UProceduralFoliageSpawner* InFoliageSpawner, const int32 RandomSeed, const int32 MaxNumSteps, const int32 InLastCancel);
 	void RemoveInstances();
 	void RemoveInstance(FProceduralFoliageInstance* Inst);
 
@@ -34,7 +30,7 @@ class FOLIAGE_API UProceduralFoliageTile : public UObject
 
 	virtual void BeginDestroy() override;
 
-	void InitSimulation(const UProceduralFoliage* ProceduralFoliage, const int32 RandomSeed);
+	void InitSimulation(const UProceduralFoliageSpawner* InFoliageSpawner, const int32 RandomSeed);
 
 	void InstancesToArray();
 
@@ -62,6 +58,9 @@ private:
 	
 private:
 
+	UPROPERTY()
+	const UProceduralFoliageSpawner* FoliageSpawner;
+
 	TSet<FProceduralFoliageInstance*> PendingRemovals;
 	TSet<FProceduralFoliageInstance*> Instances;
 
diff --git a/Engine/Source/Runtime/Foliage/Public/ProceduralFoliageVolume.h b/Engine/Source/Runtime/Foliage/Public/ProceduralFoliageVolume.h
--- a/Engine/Source/Runtime/Foliage/Public/ProceduralFoliageVolume.h
+++ b/Engine/Source/Runtime/Foliage/Public/ProceduralFoliageVolume.h
@@ -11,7 +11,7 @@ class FOLIAGE_API AProceduralFoliageVolume: public AVolume
 {
 	GENERATED_UCLASS_BODY()
 
-	UPROPERTY(Category = ProceduralFoliageTile, VisibleAnywhere, BlueprintReadOnly)
+	UPROPERTY(Category = ProceduralFoliage, VisibleAnywhere, BlueprintReadOnly)
 	UProceduralFoliageComponent* ProceduralComponent;
 
 #if WITH_EDITOR
