- Improved procedural foliage API (a bit more encapsulation & access restrictions)
- Commented a number of classes more fully
- Added a protective measure to ProceduralFoliageTile to prevent future changes to the ordering of elements in the quad tree from affecting the outcome of the simulation
    - Note that this introduces what should be a one-time instance of indeterminism relative to old simulations (i.e. resimulating after this CL will give slightly different results than before)

diff --git a/Engine/Source/Editor/FoliageEdit/Private/ProceduralFoliageComponentDetails.cpp b/Engine/Source/Editor/FoliageEdit/Private/ProceduralFoliageComponentDetails.cpp
--- a/Engine/Source/Editor/FoliageEdit/Private/ProceduralFoliageComponentDetails.cpp
+++ b/Engine/Source/Editor/FoliageEdit/Private/ProceduralFoliageComponentDetails.cpp
@@ -78,7 +78,7 @@ FReply FProceduralFoliageComponentDetails::OnResimulateClicked()
 
 			FScopedTransaction Transaction(NSLOCTEXT("ProceduralFoliageComponentDetails", "Resimulate_Transaction", "Procedural Foliage Simulation"));
 			TArray <FDesiredFoliageInstance> DesiredFoliageInstances;
-			if (Component->SpawnProceduralContent(DesiredFoliageInstances))
+			if (Component->GenerateProceduralContent(DesiredFoliageInstances))
 			{
 				FEdModeFoliage::AddInstances(Component->GetWorld(), DesiredFoliageInstances);
 
@@ -102,7 +102,7 @@ bool FProceduralFoliageComponentDetails::IsResimulateEnabled() const
 {
 	for(const TWeakObjectPtr<UProceduralFoliageComponent>& Component : SelectedComponents)
 	{
-		if(Component.IsValid() && Component->FoliageSpawner)
+		if(Component.IsValid() && Component->FoliageSpawner && Component->FoliageSpawner->GetTypes().Num() > 0)
 		{
 			return true;
 		}
diff --git a/Engine/Source/Editor/FoliageEdit/Private/ProceduralFoliageComponentVisualizer.cpp b/Engine/Source/Editor/FoliageEdit/Private/ProceduralFoliageComponentVisualizer.cpp
--- a/Engine/Source/Editor/FoliageEdit/Private/ProceduralFoliageComponentVisualizer.cpp
+++ b/Engine/Source/Editor/FoliageEdit/Private/ProceduralFoliageComponentVisualizer.cpp
@@ -6,38 +6,43 @@
 #include "ProceduralFoliageComponent.h"
 
 
-static const FColor	ProcTileColor(0,255,0);
-static const FColor	ProcTileOverlapColor(255, 255, 0);
+static const FColor& ProcTileColor = FColor::Yellow;
+static const FColor& ProcTileOverlapColor = FColor::Green;
 
 void FProceduralFoliageComponentVisualizer::DrawVisualization(const UActorComponent* Component, const FSceneView* View, FPrimitiveDrawInterface* PDI)
 {
 	if (const UProceduralFoliageComponent* ProcComponent = Cast<const UProceduralFoliageComponent>(Component))
 	{
-		if (ProcComponent->FoliageSpawner && ProcComponent->bShowDebugTiles == true)
+		if (ProcComponent->bShowDebugTiles && ProcComponent->FoliageSpawner)
 		{
+			const FVector TilesOrigin = ProcComponent->GetWorldPosition();
+
 			const float TileSize = ProcComponent->FoliageSpawner->TileSize;
-			const FVector TileSizeOffset(TileSize, TileSize, 0.f);
-			const FVector TileLocation = ProcComponent->GetWorldPosition();
-			const FVector OverlapV(ProcComponent->TileOverlap, ProcComponent->TileOverlap, 0.f);
-			
-			int32 XOffset, YOffset;
-			int32 NumX = 0;
-			int32 NumY = 0;
-			float HalfHeight;
+			const FVector TileSizeV(TileSize, TileSize, 0.f);
 
-			ProcComponent->GetTilesLayout(XOffset, YOffset, NumX, NumY, HalfHeight);
+			const float TileOverlap = ProcComponent->TileOverlap;
+			const FVector TileOverlapV(TileOverlap, TileOverlap, 0.f);
+			
+			FTileLayout TileLayout;
+			ProcComponent->GetTileLayout(TileLayout);
 
-			for (int32 X = 0; X < NumX; ++X)
+			// Draw each tile
+			for (int32 X = 0; X < TileLayout.NumTilesX; ++X)
 			{
-				for (int32 Y = 0; Y < NumY; ++Y)
+				for (int32 Y = 0; Y < TileLayout.NumTilesY; ++Y)
 				{
 					const FVector StartOffset(X*TileSize, Y*TileSize, 0.f);
 					
-					const FBox BoxInner(TileLocation + StartOffset, TileLocation + StartOffset + TileSizeOffset);
+					// Draw the tile without overlap
+					const FBox BoxInner(TilesOrigin + StartOffset, TilesOrigin + StartOffset + TileSizeV);
 					DrawWireBox(PDI, BoxInner, ProcTileColor, SDPG_World);
 
-					const FBox BoxOuter = BoxInner + (BoxInner.Min - OverlapV) + (BoxInner.Max + OverlapV);
-					DrawWireBox(PDI, BoxOuter, ProcTileOverlapColor, SDPG_World);
+					if (TileOverlap != 0.f)
+					{
+						// Draw the tile expanded by the overlap amount
+						const FBox BoxOuter = BoxInner + (BoxInner.Min - TileOverlapV) + (BoxInner.Max + TileOverlapV);
+						DrawWireBox(PDI, BoxOuter, ProcTileOverlapColor, SDPG_World);
+					}
 				}
 			}
 			
diff --git a/Engine/Source/Runtime/Foliage/Private/ProceduralFoliageComponent.cpp b/Engine/Source/Runtime/Foliage/Private/ProceduralFoliageComponent.cpp
--- a/Engine/Source/Runtime/Foliage/Private/ProceduralFoliageComponent.cpp
+++ b/Engine/Source/Runtime/Foliage/Private/ProceduralFoliageComponent.cpp
@@ -21,50 +21,35 @@ UProceduralFoliageComponent::UProceduralFoliageComponent(const FObjectInitialize
 #endif
 }
 
-void CopyTileInstances(const UProceduralFoliageTile* FromTile, UProceduralFoliageTile* ToTile, const FBox2D& InnerLocalAABB, const FTransform& ToLocalTM, const float Overlap)
+FBox2D GetTileRegion(const int32 X, const int32 Y, const float InnerSize, const float Overlap)
 {
-	const FBox2D OuterLocalAABB(InnerLocalAABB.Min, InnerLocalAABB.Max + FVector2D(Overlap, Overlap));
-	TArray<FProceduralFoliageInstance*> ToInstances;
-	FromTile->GetInstancesInAABB(OuterLocalAABB, ToInstances, false);
-	ToTile->AddInstances(ToInstances, ToLocalTM, InnerLocalAABB);
-}
-
-FBox2D GetTileRegion(const int32 X, const int32 Y, const int32 CountX, const int32 CountY, const float InnerSize, const float Overlap)
-{
-	const FVector2D Left(-Overlap, Overlap);
-	const FVector2D Bottom(Overlap, -Overlap);
-
-	FBox2D Region(FVector2D(Overlap,Overlap), FVector2D(InnerSize + Overlap, InnerSize + Overlap));
-	if (X == 0)
-	{
-		Region += Left;
-	}
-
-	if (Y == 0)
-	{
-		Region += Bottom;
-	}
+	// For tiles on the bottom and the left (that have no preceding neighbor to overlap them), have the
+	// tile fill that space with its own content
+	const float BottomLeftX = X == 0 ? -Overlap : Overlap;
+	const float BottomLeftY = Y == 0 ? -Overlap : Overlap;
+	FBox2D Region(FVector2D(BottomLeftX, BottomLeftY), FVector2D(InnerSize + Overlap, InnerSize + Overlap));
 
 	return Region;
 }
-void UProceduralFoliageComponent::GetTilesLayout(int32& MinXIdx, int32& MinYIdx, int32& NumX, int32& NumY, float& HalfHeight) const
+
+void UProceduralFoliageComponent::GetTileLayout(FTileLayout& OutTileLayout) const
 {
-	if(UBrushComponent* Brush = SpawningVolume->GetBrushComponent())
+	if (UBrushComponent* Brush = SpawningVolume->GetBrushComponent())
 	{
+		// Determine the bottom-left-most tile that contains the min position (when accounting for overlap)
 		const FVector MinPosition = Brush->Bounds.GetBox().Min + TileOverlap;
-		const FVector MaxPosition = Brush->Bounds.GetBox().Max - TileOverlap;
-
-		//we want to find the bottom left tile that contains this MinPosition
-		MinXIdx = FMath::FloorToInt(MinPosition.X / FoliageSpawner->TileSize);
-		MinYIdx = FMath::FloorToInt(MinPosition.Y / FoliageSpawner->TileSize);
+		OutTileLayout.BottomLeftX = FMath::FloorToInt(MinPosition.X / FoliageSpawner->TileSize);
+		OutTileLayout.BottomLeftY = FMath::FloorToInt(MinPosition.Y / FoliageSpawner->TileSize);
 
+		// Determine the total number of tiles along each active axis
+		const FVector MaxPosition = Brush->Bounds.GetBox().Max - TileOverlap;
 		const int32 MaxXIdx = FMath::FloorToInt(MaxPosition.X / FoliageSpawner->TileSize);
 		const int32 MaxYIdx = FMath::FloorToInt(MaxPosition.Y / FoliageSpawner->TileSize);
 
-		NumX = (MaxXIdx - MinXIdx) + 1;
-		NumY = (MaxYIdx - MinYIdx) + 1;
+		OutTileLayout.NumTilesX = (MaxXIdx - OutTileLayout.BottomLeftX) + 1;
+		OutTileLayout.NumTilesY = (MaxYIdx - OutTileLayout.BottomLeftY) + 1;
 
-		HalfHeight = Brush->Bounds.GetBox().GetExtent().Z;
+		OutTileLayout.HalfHeight = Brush->Bounds.GetBox().GetExtent().Z;
 	}
 }
 
@@ -76,109 +61,98 @@ FVector UProceduralFoliageComponent::GetWorldPosition() const
 		return FVector::ZeroVector;
 	}
 	
-	int32 X1 = 0;
-	int32 Y1 = 0;
-	int32 NumX, NumY;
-	float HalfHeight;
-	GetTilesLayout(X1, Y1, NumX, NumY, HalfHeight);
+	FTileLayout TileLayout;
+	GetTileLayout(TileLayout);
 
-	return FVector(X1 * FoliageSpawner->TileSize, Y1 * FoliageSpawner->TileSize, Brush->Bounds.Origin.Z);
+	const float TileSize = FoliageSpawner->TileSize;
+	return FVector(TileLayout.BottomLeftX * TileSize, TileLayout.BottomLeftY * TileSize, Brush->Bounds.Origin.Z);
 }
 
-bool UProceduralFoliageComponent::SpawnTiles(TArray<FDesiredFoliageInstance>& OutInstances)
+bool UProceduralFoliageComponent::ExecuteSimulation(TArray<FDesiredFoliageInstance>& OutInstances)
 {
 #if WITH_EDITOR
 	if (FoliageSpawner && SpawningVolume && SpawningVolume->GetBrushComponent())
 	{
-		/** Constants for laying out the overlapping tile grid*/
-		const float InnerTileSize = FoliageSpawner->TileSize;
-		const FVector2D InnerTileV(InnerTileSize, InnerTileSize);
-		const FVector2D InnerTileX(InnerTileSize, 0.f);
-		const FVector2D InnerTileY(0.f, InnerTileSize);
-
-		const float OuterTileSize = FoliageSpawner->TileSize + TileOverlap;
-		const FVector2D OuterTileV(OuterTileSize, OuterTileSize);
-		const FVector2D OuterTileX(OuterTileSize, 0.f);
-		const FVector2D OuterTileY(0.f, OuterTileSize);
-
-		const FVector2D OverlapX(TileOverlap, 0.f);
-		const FVector2D OverlapY(0.f, TileOverlap);
-
-		TArray<TFuture< TArray<FDesiredFoliageInstance>* >> Futures;
-		
-		int32 XOffset, YOffset, NumTilesX, NumTilesY;
-		float HalfHeight;
-		GetTilesLayout(XOffset, YOffset, NumTilesX, NumTilesY, HalfHeight);
-
-		FBodyInstance* VolumeBodyInstance = SpawningVolume->GetBrushComponent()->GetBodyInstance();
-
+		// Establish the counter used to check if the user has canceled the simulation
 		FThreadSafeCounter LastCancel;
 		const int32 LastCanelInit = LastCancel.GetValue();
 		FThreadSafeCounter* LastCancelPtr = &LastCancel;
-
+		
+		// Establish basic info about the tiles
+		const float TileSize = FoliageSpawner->TileSize;
 		const FVector WorldPosition = GetWorldPosition();
+		FTileLayout TileLayout;
+		GetTileLayout(TileLayout);
+
+		FBodyInstance* VolumeBodyInstance = SpawningVolume->GetBrushComponent()->GetBodyInstance();
+
 		FoliageSpawner->SimulateIfNeeded();
 
-		for (int32 X = 0; X < NumTilesX; ++X)
+		TArray<TFuture< TArray<FDesiredFoliageInstance>* >> Futures;
+		for (int32 X = 0; X < TileLayout.NumTilesX; ++X)
 		{
-			for (int32 Y = 0; Y < NumTilesY; ++Y)
+			for (int32 Y = 0; Y < TileLayout.NumTilesY; ++Y)
 			{
-				//We have to get the tiles and create new one to build on main thread
-				const UProceduralFoliageTile* Tile = FoliageSpawner->GetRandomTile(X + XOffset, Y + YOffset);
-				if (Tile == nullptr)	//simulation was canceled or failed
+				// We have to get the simulated tiles and create new ones to build on main thread
+				const UProceduralFoliageTile* Tile = FoliageSpawner->GetRandomTile(X + TileLayout.BottomLeftX, Y + TileLayout.BottomLeftY);
+				if (Tile == nullptr)	
 				{
+					// Simulation was either canceled or failed
 					return false;
 				}
 
-				const UProceduralFoliageTile* RightTile = ( X + 1 < NumTilesX ) ? FoliageSpawner->GetRandomTile(X + XOffset + 1, Y + YOffset) : nullptr;
-				const UProceduralFoliageTile* TopTile = ( Y + 1 < NumTilesY ) ? FoliageSpawner->GetRandomTile(X + XOffset, Y + YOffset + 1) : nullptr;
-				const UProceduralFoliageTile* TopRightTile = ( RightTile && TopTile ) ? FoliageSpawner->GetRandomTile(X + XOffset + 1, Y + YOffset + 1) : nullptr;
+				// From the pool of simulated tiles, pick the neighbors of this tile
+				const UProceduralFoliageTile* RightTile = (X + 1 < TileLayout.NumTilesX) ? FoliageSpawner->GetRandomTile(X + TileLayout.BottomLeftX + 1, Y + TileLayout.BottomLeftY) : nullptr;
+				const UProceduralFoliageTile* TopTile = (Y + 1 < TileLayout.NumTilesY) ? FoliageSpawner->GetRandomTile(X + TileLayout.BottomLeftX, Y + TileLayout.BottomLeftY + 1) : nullptr;
+				const UProceduralFoliageTile* TopRightTile = (RightTile && TopTile) ? FoliageSpawner->GetRandomTile(X + TileLayout.BottomLeftX + 1, Y + TileLayout.BottomLeftY + 1) : nullptr;
 
-				UProceduralFoliageTile* JTile = FoliageSpawner->CreateTempTile();
+				// Create a temp tile that will contain the composite contents of the tile after accounting for overlap
+				UProceduralFoliageTile* CompositeTile = FoliageSpawner->CreateTempTile();
 
 				Futures.Add(Async<TArray<FDesiredFoliageInstance>*>(EAsyncExecution::ThreadPool, [=]()
 				{
 					if (LastCancelPtr->GetValue() != LastCanelInit)
 					{
+						// The counter has changed since we began, meaning the user canceled the operation
 						return new TArray<FDesiredFoliageInstance>();
 					}
 
-					const FVector OrientedOffset = FVector(X, Y, 0.f) * FVector(InnerTileSize, InnerTileSize, 0.f);
-					const FTransform TileTM(OrientedOffset + WorldPosition);
+					//@todo proc foliage: Determine the composite contents of the tile (including overlaps) without copying everything to a temp tile
 
-					//copy the inner tile
-					const FBox2D InnerBox = GetTileRegion(X, Y, NumTilesX, NumTilesY, InnerTileSize, TileOverlap);
-					CopyTileInstances(Tile, JTile, InnerBox, FTransform::Identity, TileOverlap);
+					// Copy the base tile contents
+					const FBox2D BaseTile = GetTileRegion(X, Y, TileSize, TileOverlap);
+					Tile->CopyInstancesToTile(CompositeTile, BaseTile, FTransform::Identity, TileOverlap);
 
 
 					if (RightTile)
 					{
-						//Add overlap to the right
-						const FBox2D RightBox(FVector2D(-TileOverlap, InnerBox.Min.Y), FVector2D(TileOverlap, InnerBox.Max.Y));
-						const FTransform RightTM(FVector(InnerTileSize, 0.f, 0.f));
-						CopyTileInstances(RightTile, JTile, RightBox, RightTM, TileOverlap);
+						// Add instances from the right overlapping tile
+						const FBox2D RightBox(FVector2D(0.f, BaseTile.Min.Y), FVector2D(TileOverlap, BaseTile.Max.Y));
+						const FTransform RightTM(FVector(TileSize, 0.f, 0.f));
+						RightTile->CopyInstancesToTile(CompositeTile, RightBox, RightTM, TileOverlap);
 					}
 
 					if (TopTile)
 					{
-						//Add overlap to the top
-						const FBox2D TopBox(FVector2D(InnerBox.Min.X, -TileOverlap), FVector2D(InnerBox.Max.X, TileOverlap));
-						const FTransform TopTM(FVector(0.f, InnerTileSize, 0.f));
-						CopyTileInstances(TopTile, JTile, TopBox, TopTM, TileOverlap);
+						// Add instances from the top overlapping tile
+						const FBox2D TopBox(FVector2D(BaseTile.Min.X, -TileOverlap), FVector2D(BaseTile.Max.X, TileOverlap));
+						const FTransform TopTM(FVector(0.f, TileSize, 0.f));
+						TopTile->CopyInstancesToTile(CompositeTile, TopBox, TopTM, TileOverlap);
 					}
 
 					if (TopRightTile)
 					{
-						//Add overlap to the top right
+						// Add instances from the top-right overlapping tile
 						const FBox2D TopRightBox(FVector2D(-TileOverlap, -TileOverlap), FVector2D(TileOverlap, TileOverlap));
-						const FTransform TopRightTM(FVector(InnerTileSize, InnerTileSize, 0.f));
-						CopyTileInstances(TopRightTile, JTile, TopRightBox, TopRightTM, TileOverlap);
+						const FTransform TopRightTM(FVector(TileSize, TileSize, 0.f));
+						TopRightTile->CopyInstancesToTile(CompositeTile, TopRightBox, TopRightTM, TileOverlap);
 					}
 
+					const FVector OrientedOffset = FVector(X, Y, 0.f) * TileSize;
+					const FTransform TileTM(OrientedOffset + WorldPosition);
+
 					TArray<FDesiredFoliageInstance>* DesiredInstances = new TArray<FDesiredFoliageInstance>();
-					JTile->InstancesToArray();
-					JTile->CreateInstancesToSpawn(*DesiredInstances, TileTM, ProceduralGuid, HalfHeight, VolumeBodyInstance);
-					JTile->Empty();
+					CompositeTile->ExtractDesiredInstances(*DesiredInstances, TileTM, ProceduralGuid, TileLayout.HalfHeight, VolumeBodyInstance);
 
 					return DesiredInstances;
 					
@@ -188,32 +162,33 @@ bool UProceduralFoliageComponent::SpawnTiles(TArray<FDesiredFoliageInstance>& Ou
 		}
 
 		const FText StatusMessage = LOCTEXT("PlaceProceduralFoliage", "Placing ProceduralFoliage...");
-		const FText CancelMessage = LOCTEXT("PlaceProceduralFoliageCancel", "Cancelling ProceduralFoliage...");
+		const FText CancelMessage = LOCTEXT("PlaceProceduralFoliageCancel", "Canceling ProceduralFoliage...");
 		GWarn->BeginSlowTask(StatusMessage, true, true);
 
 
 		int32 FutureIdx = 0;
 		bool bCancelled = false;
-		for (int X = 0; X < NumTilesX; ++X)
+		for (int X = 0; X < TileLayout.NumTilesX; ++X)
 		{
-			for (int Y = 0; Y < NumTilesY; ++Y)
+			for (int Y = 0; Y < TileLayout.NumTilesY; ++Y)
 			{
 				bool bFirstTime = true;
 				while (Futures[FutureIdx].WaitFor(FTimespan(0, 0, 0, 0, 100)) == false || bFirstTime)
 				{
 					if (GWarn->ReceivedUserCancel() && bCancelled == false)
 					{
+						// Increment the thread-safe counter. Tiles compare against the original count and cancel if different.
 						LastCancel.Increment();
 						bCancelled = true;
 					}
 
 					if (bCancelled)
 					{
-						GWarn->StatusUpdate(Y + X * NumTilesY, NumTilesX * NumTilesY, CancelMessage);
+						GWarn->StatusUpdate(Y + X * TileLayout.NumTilesY, TileLayout.NumTilesX * TileLayout.NumTilesY, CancelMessage);
 					}
 					else
 					{
-						GWarn->StatusUpdate(Y + X * NumTilesY, NumTilesX * NumTilesY, StatusMessage);
+						GWarn->StatusUpdate(Y + X * TileLayout.NumTilesY, TileLayout.NumTilesX * TileLayout.NumTilesY, StatusMessage);
 					}
 
 					bFirstTime = false;
@@ -234,11 +209,11 @@ bool UProceduralFoliageComponent::SpawnTiles(TArray<FDesiredFoliageInstance>& Ou
 
 }
 
-bool UProceduralFoliageComponent::SpawnProceduralContent(TArray <FDesiredFoliageInstance>& OutInstances)
+bool UProceduralFoliageComponent::GenerateProceduralContent(TArray <FDesiredFoliageInstance>& OutInstances)
 {
 #if WITH_EDITOR
 	
-	if (SpawnTiles(OutInstances))
+	if (ExecuteSimulation(OutInstances))
 	{
 		RemoveProceduralContent();
 		return true;
diff --git a/Engine/Source/Runtime/Foliage/Private/ProceduralFoliageSpawner.cpp b/Engine/Source/Runtime/Foliage/Private/ProceduralFoliageSpawner.cpp
--- a/Engine/Source/Runtime/Foliage/Private/ProceduralFoliageSpawner.cpp
+++ b/Engine/Source/Runtime/Foliage/Private/ProceduralFoliageSpawner.cpp
@@ -130,11 +130,15 @@ const UProceduralFoliageTile* UProceduralFoliageSpawner::GetRandomTile(int32 X,
 {
 	if (PrecomputedTiles.Num())
 	{
-		FRandomStream HashStream;	//using this as a hash function
+		// Random stream to use as a hash function
+		FRandomStream HashStream;	
+		
 		HashStream.Initialize(X);
 		const float XRand = HashStream.FRand();
+		
 		HashStream.Initialize(Y);
 		const float YRand = HashStream.FRand();
+		
 		const int32 RandomNumber = (RAND_MAX * XRand / (YRand + 0.01f));
 		const int32 Idx = FMath::Clamp(RandomNumber % PrecomputedTiles.Num(), 0, PrecomputedTiles.Num() - 1);
 		return PrecomputedTiles[Idx].Get();
@@ -175,17 +179,21 @@ void UProceduralFoliageSpawner::Simulate(int32 NumSteps)
 
 	for (int32 FutureIdx = 0; FutureIdx < Futures.Num(); ++FutureIdx)
 	{
-		while (Futures[FutureIdx].WaitFor(FTimespan(0, 0, 0, 0, 100)) == false)		//sleep for 100ms if not ready. Needed so cancel is responsive
+		// Sleep for 100ms if not ready. Needed so cancel is responsive.
+		while (Futures[FutureIdx].WaitFor(FTimespan(0, 0, 0, 0, 100)) == false)
 		{
 			GWarn->StatusUpdate(FutureIdx, TotalTasks, LOCTEXT("SimulateProceduralFoliage", "Simulate ProceduralFoliage..."));
-			if (GWarn->ReceivedUserCancel() && bCancelled == false)	//If the user wants to cancel just increment this. Tiles compare against the original count and cancel if needed
+
+			if (GWarn->ReceivedUserCancel() && bCancelled == false)
 			{
+				// Increment the thread-safe counter. Tiles compare against the original count and cancel if different.
 				LastCancel.Increment();
 				bCancelled = true;
 			}
 		}
 
-		PrecomputedTiles.Add(Futures[FutureIdx].Get());		//Even if we cancel we block until threads have exited safely to ensure memory isn't GCed
+		// Even if canceled, block until all threads have exited safely. This ensures memory isn't GC'd.
+		PrecomputedTiles.Add(Futures[FutureIdx].Get());		
 	}
 
 	GWarn->EndSlowTask();
diff --git a/Engine/Source/Runtime/Foliage/Private/ProceduralFoliageTile.cpp b/Engine/Source/Runtime/Foliage/Private/ProceduralFoliageTile.cpp
--- a/Engine/Source/Runtime/Foliage/Private/ProceduralFoliageTile.cpp
+++ b/Engine/Source/Runtime/Foliage/Private/ProceduralFoliageTile.cpp
@@ -16,16 +16,15 @@ UProceduralFoliageTile::UProceduralFoliageTile(const FObjectInitializer& ObjectI
 
 bool UProceduralFoliageTile::HandleOverlaps(FProceduralFoliageInstance* Instance)
 {
-	//We now check what the instance overlaps. If any of its overlaps dominate we remove the instance and leave everything else alone.
-	//If the instance survives we mark all dominated overlaps as pending removal. They will be removed from the broadphase and will not spread seeds or age.
-	//Note that this introduces potential indeterminism! If the iteration order changes we could get different results. This is needed because it gives us huge performance savings.
-	//Note that if the underlying data structures stay the same (i.e. no core engine changes) this should not matter. This gives us short term determinism, but not long term.
+	// If the instance survives we mark all dominated overlaps as pending removal. They will be removed from the broadphase and will not spread seeds or age.
+	// Note that this introduces potential indeterminism! If the iteration order changes we could get different results. This is needed because it gives us huge performance savings.
+	// Note that if the underlying data structures stay the same (i.e. no core engine changes) this should not matter. This gives us short term determinism, but not long term.
 
 	bool bSurvived = true;
 	TArray<FProceduralFoliageOverlap> Overlaps;
 	Broadphase.GetOverlaps(Instance, Overlaps);
 
-	//Check if the instance survives
+	// Check if the instance survives
 	for (const FProceduralFoliageOverlap& Overlap : Overlaps)
 	{
 		FProceduralFoliageInstance* Dominated = FProceduralFoliageInstance::Domination(Overlap.A, Overlap.B, Overlap.OverlapType);
@@ -42,8 +41,11 @@ bool UProceduralFoliageTile::HandleOverlaps(FProceduralFoliageInstance* Instance
 		{
 			if (FProceduralFoliageInstance* Dominated = FProceduralFoliageInstance::Domination(Overlap.A, Overlap.B, Overlap.OverlapType))
 			{
-				check(Dominated != Instance);	//we shouldn't be here if we survived
-				MarkPendingRemoval(Dominated);	//We can't immediately remove because we're potentially iterating over existing instances.
+				// Should only be here if we didn't survive
+				check(Dominated != Instance);
+
+				//We can't immediately remove because we're potentially iterating over existing instances.
+				MarkPendingRemoval(Dominated);
 			}
 		}
 	}
@@ -69,8 +71,7 @@ FProceduralFoliageInstance* UProceduralFoliageTile::NewSeed(const FVector& Locat
 		NewInst->Scale = Scale;
 		NewInst->bBlocker = bBlocker;
 		
-	
-
+		// Add the seed if possible
 		Broadphase.Insert(NewInst);
 		const bool bSurvived = HandleOverlaps(NewInst);
 		return bSurvived ? NewInst : nullptr;
@@ -123,7 +124,7 @@ FVector UProceduralFoliageTile::GetSeedOffset(const UFoliageType_InstancedStatic
 void UProceduralFoliageTile::AgeSeeds()
 {
 	TArray<FProceduralFoliageInstance*> NewSeeds;
-	for (FProceduralFoliageInstance* Instance : Instances)
+	for (FProceduralFoliageInstance* Instance : InstancesSet)
 	{
 		if (UserCancelled()){ return; }
 		if (Instance->IsAlive())
@@ -137,6 +138,7 @@ void UProceduralFoliageTile::AgeSeeds()
 
 				const FVector Location = Instance->Location;
 
+				// Replace the current instance with the newly aged version
 				MarkPendingRemoval(Instance);
 				if (FProceduralFoliageInstance* Inst = NewSeed(Location, NewScale, Type, NewAge))
 				{
@@ -146,22 +148,25 @@ void UProceduralFoliageTile::AgeSeeds()
 		}
 	}
 
+	// Save all the newly created aged instances
 	for (FProceduralFoliageInstance* Seed : NewSeeds)
 	{
-		Instances.Add(Seed);
+		InstancesSet.Add(Seed);
 	}
 
+	// Get rid of the old younger versions
 	FlushPendingRemovals();
-	
 }
 
 void UProceduralFoliageTile::SpreadSeeds(TArray<FProceduralFoliageInstance*>& NewSeeds)
 {
-	for (FProceduralFoliageInstance* Inst : Instances)
+	for (FProceduralFoliageInstance* Inst : InstancesSet)
 	{
 		if (UserCancelled()){ return; }
-		if (Inst->IsAlive() == false)	//The instance has been killed so don't bother spreading seeds. Note this introduces potential indeterminism if the order of instance traversal changes (implementation details of TSet for example)
+		if (Inst->IsAlive() == false)
 		{
+			// The instance has been killed so don't bother spreading seeds. 
+			// Note this introduces potential indeterminism if the order of instance traversal changes (implementation details of TSet for example)
 			continue;
 		}
 
@@ -180,9 +185,9 @@ void UProceduralFoliageTile::SpreadSeeds(TArray<FProceduralFoliageInstance*>& Ne
 				if (GlobalOffset.SizeSquared2D() + SMALL_NUMBER > MinDistanceToClear*MinDistanceToClear)
 				{
 					const FVector NewLocation = GlobalOffset + Inst->Location;
-					if (FProceduralFoliageInstance* Inst = NewSeed(NewLocation, NewScale, Type, NewAge))
+					if (FProceduralFoliageInstance* NewInstance = NewSeed(NewLocation, NewScale, Type, NewAge))
 					{
-						NewSeeds.Add(Inst);
+						NewSeeds.Add(NewInstance);
 					}
 				}
 			}
@@ -321,7 +326,7 @@ void UProceduralFoliageTile::MarkPendingRemoval(FProceduralFoliageInstance* ToRe
 
 void UProceduralFoliageTile::RemoveInstances()
 {
-	for (FProceduralFoliageInstance* Inst : Instances)
+	for (FProceduralFoliageInstance* Inst : InstancesSet)
 	{
 		MarkPendingRemoval(Inst);
 	}
@@ -332,10 +337,11 @@ void UProceduralFoliageTile::RemoveInstances()
 
 void UProceduralFoliageTile::InstancesToArray()
 {
-	InstancesArray.Empty(Instances.Num());
-	for (FProceduralFoliageInstance* FromInst : Instances)
+	InstancesArray.Empty(InstancesSet.Num());
+	for (FProceduralFoliageInstance* FromInst : InstancesSet)
 	{
-		if (FromInst->bBlocker == false)	//blockers do not get instantiated so don't bother putting it into array
+		// Blockers do not get instantiated so don't bother putting it into array
+		if (FromInst->bBlocker == false)
 		{
 			new(InstancesArray)FProceduralFoliageInstance(*FromInst);
 		}
@@ -350,7 +356,7 @@ void UProceduralFoliageTile::RemoveInstance(FProceduralFoliageInstance* ToRemove
 		ToRemove->TerminateInstance();
 	}
 	
-	Instances.Remove(ToRemove);
+	InstancesSet.Remove(ToRemove);
 	delete ToRemove;
 }
 
@@ -394,13 +400,13 @@ void UProceduralFoliageTile::StepSimulation()
 
 	for (FProceduralFoliageInstance* Inst : NewInstances)
 	{
-		Instances.Add(Inst);
+		InstancesSet.Add(Inst);
 	}
 
 	FlushPendingRemovals();
 }
 
-void UProceduralFoliageTile::StartSimulation(const int32 MaxNumSteps, bool bShadeGrowth)
+void UProceduralFoliageTile::RunSimulation(const int32 MaxNumSteps, bool bShadeGrowth)
 {
 	int32 MaxSteps = 0;
 
@@ -434,8 +440,8 @@ void UProceduralFoliageTile::Simulate(const UProceduralFoliageSpawner* InFoliage
 	LastCancel = InLastCancel;
 	InitSimulation(InFoliageSpawner, RandomSeed);
 
-	StartSimulation(MaxNumSteps, false);
-	StartSimulation(MaxNumSteps, true);
+	RunSimulation(MaxNumSteps, false);
+	RunSimulation(MaxNumSteps, true);
 }
 
 
@@ -445,12 +451,14 @@ void UProceduralFoliageTile::BeginDestroy()
 	RemoveInstances();
 }
 
-void UProceduralFoliageTile::CreateInstancesToSpawn(TArray<FDesiredFoliageInstance>& OutInstances, const FTransform& WorldTM, const FGuid& ProceduralGuid, const float HalfHeight, const FBodyInstance* VolumeBodyInstance) const
+void UProceduralFoliageTile::ExtractDesiredInstances(TArray<FDesiredFoliageInstance>& OutInstances, const FTransform& WorldTM, const FGuid& ProceduralGuid, const float HalfHeight, const FBodyInstance* VolumeBodyInstance, bool bEmptyTileInfo)
 {
+	InstancesToArray();
+
 	const FCollisionQueryParams Params(true);
 	FHitResult Hit;
 
-	OutInstances.Reserve(Instances.Num());
+	OutInstances.Reserve(InstancesSet.Num());
 	for (const FProceduralFoliageInstance& Instance : InstancesArray)
 	{
 		FVector StartRay = Instance.Location + WorldTM.GetLocation();
@@ -466,6 +474,25 @@ void UProceduralFoliageTile::CreateInstancesToSpawn(TArray<FDesiredFoliageInstan
 		DesiredInst->ProceduralVolumeBodyInstance = VolumeBodyInstance;
 		DesiredInst->PlacementMode = EFoliagePlacementMode::Procedural;
 	}
+
+	if (bEmptyTileInfo)
+	{
+		Empty();
+	}
+}
+
+void UProceduralFoliageTile::CopyInstancesToTile(UProceduralFoliageTile* ToTile, const FBox2D& LocalAABB, const FTransform& RelativeTM, const float Overlap) const
+{
+	//@todo proc foliage: Would be better to use the max radius of any instances in the tile instead of overlap to define the outer AABB
+
+	TArray<FProceduralFoliageInstance*> InstancesIncludingOverlap;
+	const FBox2D OuterLocalAABB(LocalAABB.Min, LocalAABB.Max + Overlap);
+	
+	// Get all the instances in the outer AABB (so we include potential blockers)
+	GetInstancesInAABB(OuterLocalAABB, InstancesIncludingOverlap);
+
+
+	ToTile->AddInstances(InstancesIncludingOverlap, RelativeTM, LocalAABB);
 }
 
 void UProceduralFoliageTile::Empty()
@@ -473,19 +500,19 @@ void UProceduralFoliageTile::Empty()
 	Broadphase.Empty();
 	InstancesArray.Empty();
 	
-	for (FProceduralFoliageInstance* Inst : Instances)
+	for (FProceduralFoliageInstance* Inst : InstancesSet)
 	{
 		delete Inst;
 	}
 
-	Instances.Empty();
+	InstancesSet.Empty();
 	PendingRemovals.Empty();
 }
 
 SIZE_T UProceduralFoliageTile::GetResourceSize(EResourceSizeMode::Type Mode)
 {
 	SIZE_T TotalSize = 0;
-	for (FProceduralFoliageInstance* Inst : Instances)
+	for (FProceduralFoliageInstance* Inst : InstancesSet)
 	{
 		TotalSize += sizeof(FProceduralFoliageInstance);
 	}
@@ -495,44 +522,51 @@ SIZE_T UProceduralFoliageTile::GetResourceSize(EResourceSizeMode::Type Mode)
 }
 
 
-void UProceduralFoliageTile::GetInstancesInAABB(const FBox2D& LocalAABB, TArray<FProceduralFoliageInstance*>& OutInstances, bool bOnTheBorder) const
+void UProceduralFoliageTile::GetInstancesInAABB(const FBox2D& LocalAABB, TArray<FProceduralFoliageInstance*>& OutInstances, bool bFullyContainedOnly) const
 {
 	TArray<FProceduralFoliageInstance*> InstancesInAABB;
 	Broadphase.GetInstancesInBox(LocalAABB, InstancesInAABB);
 
 	OutInstances.Reserve(OutInstances.Num() + InstancesInAABB.Num());
 	for (FProceduralFoliageInstance* Inst : InstancesInAABB)
-{
+	{
 		const float Rad = Inst->GetMaxRadius();
 		const FVector& Location = Inst->Location;
 
-		if (bOnTheBorder || (Location.X - Rad >= LocalAABB.Min.X && Location.X + Rad <= LocalAABB.Max.X && Location.Y - Rad >= LocalAABB.Min.Y && Location.Y + Rad <= LocalAABB.Max.Y))
+		if (!bFullyContainedOnly || (Location.X - Rad >= LocalAABB.Min.X && Location.X + Rad <= LocalAABB.Max.X && Location.Y - Rad >= LocalAABB.Min.Y && Location.Y + Rad <= LocalAABB.Max.Y))
 		{
 			OutInstances.Add(Inst);
 		}
 	}
+
+	// Sort the instances by location.
+	// This protects us from any future modifications made to the broadphase that would impact the order in which instances are located in the AABB.
+	OutInstances.Sort([](const FProceduralFoliageInstance& A, const FProceduralFoliageInstance& B)
+	{
+		return (B.Location.X == A.Location.X) ? (B.Location.Y > A.Location.Y) : (B.Location.X > A.Location.X);
+	});
 }
 
 void UProceduralFoliageTile::AddInstances(const TArray<FProceduralFoliageInstance*>& NewInstances, const FTransform& RelativeTM, const FBox2D& InnerLocalAABB)
 {
 	for (const FProceduralFoliageInstance* Inst : NewInstances)
 	{
-		const FVector& Location = Inst->Location;	//we need the local space because we're comparing it to the AABB
-		//Instances in InnerLocalAABB or on the border of the max sides of the AABB will be visible and instantiated by this tile
-		//Instances outside of the InnerLocalAABB are only used for rejection purposes. This is needed for overlapping tiles
+		// We need the local space because we're comparing it to the AABB
+		const FVector& Location = Inst->Location;	
 		const float Radius = Inst->GetMaxRadius();
-		bool bBlocker = false;
-		if (Location.X + Radius <= InnerLocalAABB.Min.X || Location.X - Radius > InnerLocalAABB.Max.X || Location.Y + Radius <= InnerLocalAABB.Min.Y || Location.Y - Radius > InnerLocalAABB.Max.Y)
-		{
-			//Only used for blocking instances in case of overlap, a different tile will instantiate it
-			bBlocker = true;
-		}
+		
+		// Instances in InnerLocalAABB or on the border of the max sides of the AABB will be visible and instantiated by this tile
+		// Instances outside of the InnerLocalAABB are only used for rejection purposes. This is needed for overlapping tiles
+		// The actual instantiation of the object will be taken care of by a different tile
+		const bool bIsOutsideInnerLocalAABB = Location.X + Radius <= InnerLocalAABB.Min.X
+											|| Location.X - Radius > InnerLocalAABB.Max.X
+											|| Location.Y + Radius <= InnerLocalAABB.Min.Y
+											|| Location.Y - Radius > InnerLocalAABB.Max.Y;
 
 		const FVector NewLocation = RelativeTM.TransformPosition(Inst->Location);
-		if (FProceduralFoliageInstance* NewInst = NewSeed(NewLocation, Inst->Scale, Inst->Type, Inst->Age, bBlocker))
-		{
-			
-			Instances.Add(NewInst);
+		if (FProceduralFoliageInstance* NewInst = NewSeed(NewLocation, Inst->Scale, Inst->Type, Inst->Age, bIsOutsideInnerLocalAABB))
+		{			
+			InstancesSet.Add(NewInst);
 		}
 	}
 
diff --git a/Engine/Source/Runtime/Foliage/Private/ProceduralFoliageVolume.cpp b/Engine/Source/Runtime/Foliage/Private/ProceduralFoliageVolume.cpp
--- a/Engine/Source/Runtime/Foliage/Private/ProceduralFoliageVolume.cpp
+++ b/Engine/Source/Runtime/Foliage/Private/ProceduralFoliageVolume.cpp
@@ -9,14 +9,16 @@ AProceduralFoliageVolume::AProceduralFoliageVolume(const FObjectInitializer& Obj
 : Super(ObjectInitializer)
 {
 	ProceduralComponent = ObjectInitializer.CreateDefaultSubobject<UProceduralFoliageComponent>(this, TEXT("ProceduralFoliageComponent"));
-	ProceduralComponent->SpawningVolume = this;
+	ProceduralComponent->SetSpawningVolume(this);
 
 	if (UBrushComponent* BrushComponent = GetBrushComponent())
 	{
 		BrushComponent->SetCollisionObjectType(ECC_WorldStatic);
 		BrushComponent->SetCollisionResponseToAllChannels(ECR_Ignore);
-		//This is important because the volume overlaps with all procedural foliage, which means during streaming we'll get a huge hitch for UpdateOverlaps
-		BrushComponent->bGenerateOverlapEvents = false;	//We don't care about overlaps.
+
+		// This is important because the volume overlaps with all procedural foliage
+		// That means during streaming we'll get a huge hitch for UpdateOverlaps
+		BrushComponent->bGenerateOverlapEvents = false;
 	}
 }
 
diff --git a/Engine/Source/Runtime/Foliage/Public/ProceduralFoliageComponent.h b/Engine/Source/Runtime/Foliage/Public/ProceduralFoliageComponent.h
--- a/Engine/Source/Runtime/Foliage/Public/ProceduralFoliageComponent.h
+++ b/Engine/Source/Runtime/Foliage/Public/ProceduralFoliageComponent.h
@@ -9,6 +9,30 @@ class UProceduralFoliageSpawner;
 class AProceduralFoliageLevelInfo;
 struct FDesiredFoliageInstance;
 
+/** Describes the layout of the tiles used for procedural foliage simulation */
+struct FTileLayout
+{
+	FTileLayout()
+		: BottomLeftX(0), BottomLeftY(0), NumTilesX(0), NumTilesY(0), HalfHeight(0.f)
+	{
+	}
+
+	// The X coordinate (in whole tiles) of the bottom-left-most active tile
+	int32 BottomLeftX;
+
+	// The Y coordinate (in whole tiles) of the bottom-left-most active tile
+	int32 BottomLeftY;
+
+	// The total number of active tiles along the x-axis
+	int32 NumTilesX;
+
+	// The total number of active tiles along the y-axis
+	int32 NumTilesY;
+	
+
+	float HalfHeight;
+};
+
 UCLASS(BlueprintType)
 class FOLIAGE_API UProceduralFoliageComponent : public UActorComponent
 {
@@ -28,21 +52,35 @@ class FOLIAGE_API UProceduralFoliageComponent : public UActorComponent
 	bool bShowDebugTiles;
 #endif
 
-	bool SpawnProceduralContent(TArray<FDesiredFoliageInstance>& OutFoliageInstances);
+	/** 
+	 * Runs the procedural foliage simulation to generate a list of desired instances to spawn.
+	 * @return True if the simulation succeeded
+	 */
+	bool GenerateProceduralContent(TArray<FDesiredFoliageInstance>& OutFoliageInstances);
+	
+	/** Removes all spawned foliage instances in the level that were spawned by this component */
 	void RemoveProceduralContent();
+
+	/** @return True if any foliage instances in the level were spawned by this component */
 	bool HasSpawnedAnyInstances();
-	const FGuid& GetProceduralGuid() const { return ProceduralGuid;  }
+	
+	/** @return The position in world space of the bottom-left corner of the bottom-left-most active tile */
+	FVector GetWorldPosition() const;
 
-	UPROPERTY()
-	AVolume* SpawningVolume;
+	/** Determines the basic layout of the tiles used in the simulation */
+	void GetTileLayout(FTileLayout& OutTileLayout) const;
 
-	FVector GetWorldPosition() const;
-	void GetTilesLayout(int32& MinXIdx, int32& MinYIdx, int32& NumXTiles, int32& NumYTiles, float& HalfHeight) const;
+	void SetSpawningVolume(AVolume* InSpawningVolume) { SpawningVolume = InSpawningVolume; }
+	const FGuid& GetProceduralGuid() const { return ProceduralGuid; }
 
 private:
-	bool SpawnTiles(TArray<FDesiredFoliageInstance>& OutFoliageInstances);
+	/** Does all the actual work of executing the procedural foliage simulation */
+	bool ExecuteSimulation(TArray<FDesiredFoliageInstance>& OutFoliageInstances);
 
 private:
+	UPROPERTY()
+	AVolume* SpawningVolume;
+	
 	UPROPERTY()
 	FGuid ProceduralGuid;
 };
\ No newline at end of file
diff --git a/Engine/Source/Runtime/Foliage/Public/ProceduralFoliageSpawner.h b/Engine/Source/Runtime/Foliage/Public/ProceduralFoliageSpawner.h
--- a/Engine/Source/Runtime/Foliage/Public/ProceduralFoliageSpawner.h
+++ b/Engine/Source/Runtime/Foliage/Public/ProceduralFoliageSpawner.h
@@ -29,20 +29,22 @@ class FOLIAGE_API UProceduralFoliageSpawner : public UObject
 {
 	GENERATED_UCLASS_BODY()
 
-	/**The seed used for generating the randomness of the simulation. */
+	/** The seed used for generating the randomness of the simulation. */
 	UPROPERTY(Category = ProceduralFoliageSimulation, EditAnywhere, BlueprintReadOnly)
 	int32 RandomSeed;
 
-	/**Cm size of the tile along one axis. The tile is square so the total area is TileSize*TileSize */
+	/** Length of the tile (in cm) along one axis. The total area of the tile will be TileSize*TileSize. */
 	UPROPERTY(Category = ProceduralFoliageSimulation, EditAnywhere, BlueprintReadOnly)
 	float TileSize;
 
+	/** The number of unique tiles to generate. The final simulation is a procedurally determined combination of the various unique tiles. */
 	UPROPERTY(Category = ProceduralFoliageSimulation, EditAnywhere, BlueprintReadOnly)
 	int32 NumUniqueTiles;
 
 	FThreadSafeCounter LastCancel;
 
 private:
+	/** The types of foliage to procedurally spawn. */
 	UPROPERTY(Category = ProceduralFoliageSimulation, EditAnywhere)
 	TArray<FProceduralFoliageTypeData> Types;
 
diff --git a/Engine/Source/Runtime/Foliage/Public/ProceduralFoliageTile.h b/Engine/Source/Runtime/Foliage/Public/ProceduralFoliageTile.h
--- a/Engine/Source/Runtime/Foliage/Public/ProceduralFoliageTile.h
+++ b/Engine/Source/Runtime/Foliage/Public/ProceduralFoliageTile.h
@@ -12,34 +12,95 @@ struct FProceduralFoliageOverlap;
 class UProceduralFoliageSpawner;
 class UFoliageType_InstancedStaticMesh;
 
+/**
+ *	Procedurally determines where to spawn foliage meshes within a discrete area.
+ *	Generally, a procedural foliage simulation as a whole is composed of multiple tiles.
+ *	Tiles are able to overlap one another as well to create a seamless appearance.
+ *	
+ *	Note that the tile is not responsible for actually spawning any instances, it only determines where they should be placed.
+ *	Following a simulation, call ExtractDesiredInstances for information about where each instance should spawn.
+ *	
+ *	Note also that, barring any core changes to the ordering of TSet, foliage generation is deterministic 
+ *	(i.e. given the same inputs, the result of the simulation will always be the same).
+ */
 UCLASS()
 class FOLIAGE_API UProceduralFoliageTile : public UObject
 {
 	GENERATED_UCLASS_BODY()
 
+	/**
+	 *	Procedurally determine the placement of foliage instances.
+	 *	@param InFoliageSpawner the foliage spawner containing the various foliage types to simulate
+	 *	@param RandomSeed For seeding (in the RNG sense) the random dispersion of seeds (in the foliage sense)
+	 *	@param MaxNumSteps The maximum number of steps to run the simulation
+	 */
 	void Simulate(const UProceduralFoliageSpawner* InFoliageSpawner, const int32 RandomSeed, const int32 MaxNumSteps, const int32 InLastCancel);
-	void RemoveInstances();
+	
+	/**
+	 *	Extracts information on each instance that should be spawned based on the simulation.
+	 */
+	void ExtractDesiredInstances(TArray<FDesiredFoliageInstance>& OutDesiredInstances, const FTransform& WorldTM, const FGuid& ProceduralGuid, const float HalfHeight, const FBodyInstance* VolumeBodyInstance, bool bEmptyTileInfo = true);
+	
+	/**
+	 *	Copies instances within this tile to another
+	 *	@param ToTile The tile to copy instances to
+	 *	@param LocalAABB The region (in local space) to copy instances from
+	 *	@param RelativeTM The transform of the destination tile relative to this tile
+	 *	@param Overlap The amount by which the tiles are overlapping
+	 */
+	void CopyInstancesToTile(UProceduralFoliageTile* ToTile, const FBox2D& LocalAABB, const FTransform& RelativeTM, const float Overlap) const;
+
+	/**
+	 *	Removes a single instance from the tile
+	 *	@param Inst The instance to remove
+	 */
 	void RemoveInstance(FProceduralFoliageInstance* Inst);
 
-	/** Takes a local AABB and finds all the instances in it. The instances are returned in the space local to the AABB. That is, an instance in the bottom left AABB will have a coordinate of (0,0) */
-	void GetInstancesInAABB(const FBox2D& LocalAABB, TArray<FProceduralFoliageInstance* >& OutInstances, bool bOnTheBorder) const;
-	void CreateInstancesToSpawn(TArray<FDesiredFoliageInstance>& OutInstances, const FTransform& WorldTM, const FGuid& ProceduralGuid, const float HalfHeight, const FBodyInstance* VolumeBodyInstance) const;
-	void AddInstances(const TArray<FProceduralFoliageInstance*>& NewInstances, const FTransform& LocalTM, const FBox2D& InnerLocalAABB);
+	/** Removes all instances from the tile */
+	void RemoveInstances();
+
+	/** Initializes the tile for simulation. Called automatically by Simulate, so only bother to call if the tile won't actually simulate. */
+	void InitSimulation(const UProceduralFoliageSpawner* InFoliageSpawner, const int32 RandomSeed);
 
+	// UObject interface
 	virtual SIZE_T GetResourceSize(EResourceSizeMode::Type Mode) override;
-
 	virtual void BeginDestroy() override;
 
-	void InitSimulation(const UProceduralFoliageSpawner* InFoliageSpawner, const int32 RandomSeed);
-
+private:
+	
+	/**
+	* Finds all the instances within an AABB. The instances are returned in the space local to the AABB.
+	* That is, an instance in the bottom left AABB will have a coordinate of (0,0).
+	*
+	* @param LocalAABB The AABB to get instances within
+	* @param OutInstances Contains the instances within the specified AABB
+	* @param bFullyContainedOnly If true, will only get instances that are fully contained within the AABB
+	*/
+	void GetInstancesInAABB(const FBox2D& LocalAABB, TArray<FProceduralFoliageInstance* >& OutInstances, bool bFullyContainedOnly = true) const;
+
+	/**
+	 * Adds procedural instances to the tile from some external source.
+	 * @param The new instances to add to the tile
+	 * @param ToLocalTM Used to transform the new instances to this tile's local space
+	 * @param InnerLocalAABB New instances outside this region will be tracked for blocking purposes only (and never spawned)
+	 */
+	void AddInstances(const TArray<FProceduralFoliageInstance*>& NewInstances, const FTransform& ToLocalTM, const FBox2D& InnerLocalAABB);
+
+	/** Fills the InstancesArray by iterating through the contents of the InstancesSet */
 	void InstancesToArray();
 
 	/** Empty arrays and sets */
 	void Empty();
-private:
-	
-	void StartSimulation(const int32 MaxNumSteps, bool bShadeGrowth);
+
+	/** Run the simulation for the appropriate number of steps. 
+	 * @param bShadeGrowth Whether the simulation should run exclusively within existing shade from other instances
+	 */
+	void RunSimulation(const int32 MaxNumSteps, bool bShadeGrowth);
+
+	/** Advance the simulation by a step (think of a step as a generation) */
 	void StepSimulation();
+
+	/** Randomly seed the next step of the simulation */
 	void AddRandomSeeds(TArray<FProceduralFoliageInstance*>& OutInstances);
 
 	/** Determines whether the instance will survive all the overlaps, and then kills the appropriate instances. Returns true if the instance survives */
@@ -62,7 +123,7 @@ private:
 	const UProceduralFoliageSpawner* FoliageSpawner;
 
 	TSet<FProceduralFoliageInstance*> PendingRemovals;
-	TSet<FProceduralFoliageInstance*> Instances;
+	TSet<FProceduralFoliageInstance*> InstancesSet;
 
 	UPROPERTY()
 	TArray<FProceduralFoliageInstance> InstancesArray;
