Fixed: Issues selecting instance meshes with Foliage Select tool UE-13147

diff --git a/Engine/Source/Runtime/Engine/Classes/Components/InstancedStaticMeshComponent.h b/Engine/Source/Runtime/Engine/Classes/Components/InstancedStaticMeshComponent.h
--- a/Engine/Source/Runtime/Engine/Classes/Components/InstancedStaticMeshComponent.h
+++ b/Engine/Source/Runtime/Engine/Classes/Components/InstancedStaticMeshComponent.h
@@ -140,7 +140,9 @@ public:
 	/** One bit per instance if the instance is selected. */
 	TBitArray<> SelectedInstances;
 #endif
-
+	/** Incremented when instance selection is changed */
+	int32 SelectionStamp;
+	
 #if WITH_PHYSX
 	/** Aggregate physx representation of the instances' bodies. */
 	TArray<physx::PxAggregate*> Aggregates;
@@ -193,6 +195,9 @@ public:
 	/** Select/deselect an instance or group of instances */
 	void SelectInstance(bool bInSelected, int32 InInstanceIndex, int32 InInstanceCount = 1);
 
+	/** Deselect all instances */
+	void ClearInstanceSelection();
+
 	/** 
 	 * Transfers ownership of instance render data to a render thread 
 	 * instance render data will be released in scene proxy dtor or on render thread task
diff --git a/Engine/Source/Runtime/Engine/Private/InstancedStaticMesh.cpp b/Engine/Source/Runtime/Engine/Private/InstancedStaticMesh.cpp
--- a/Engine/Source/Runtime/Engine/Private/InstancedStaticMesh.cpp
+++ b/Engine/Source/Runtime/Engine/Private/InstancedStaticMesh.cpp
@@ -85,25 +85,16 @@ void FStaticMeshInstanceBuffer::Init(UInstancedStaticMeshComponent* InComponent,
 	// given instance index between reattaches
 	FRandomStream RandomStream( InComponent->InstancingRandomSeed );
 
-	FColor HitProxyColor(ForceInit);
-	bool bSelected = false;
 	for (int32 InstanceIndex = 0; InstanceIndex < NumRealInstances; InstanceIndex++)
 	{
-#if WITH_EDITOR
-		if (InHitProxies.Num() == NumRealInstances)
-		{
-			HitProxyColor = InHitProxies[InstanceIndex]->Id.GetColor();
-		}
-		// Record if the instance is selected
-		bSelected = InstanceIndex < InComponent->SelectedInstances.Num() && InComponent->SelectedInstances[InstanceIndex];
-#endif
-
 		FInstanceStream* InstanceRenderData = InstanceData->GetInstanceWriteAddress(bUseRemapTable ? InComponent->InstanceReorderTable[InstanceIndex] : InstanceIndex);
 		const FInstancedStaticMeshInstanceData& Instance = InComponent->PerInstanceSMData[InstanceIndex];
 						
-		InstanceRenderData->SetInstance(Instance.Transform, RandomStream.GetFraction(), Instance.LightmapUVBias, Instance.ShadowmapUVBias, HitProxyColor, bSelected);
+		InstanceRenderData->SetInstance(Instance.Transform, RandomStream.GetFraction(), Instance.LightmapUVBias, Instance.ShadowmapUVBias);
 	}
 
+	SetPerInstanceEditorData(InComponent, InHitProxies);
+
 	// Hide any removed instances
 	if (NumRemoved)
 	{
@@ -130,6 +121,30 @@ void FStaticMeshInstanceBuffer::InitFromPreallocatedData(UInstancedStaticMeshCom
 	InstanceData->SetAllowCPUAccess(true);		// GDC demo hack!
 }
 
+void FStaticMeshInstanceBuffer::SetPerInstanceEditorData(UInstancedStaticMeshComponent* InComponent, const TArray<TRefCountPtr<HHitProxy>>& InHitProxies)
+{
+#if WITH_EDITOR
+	if (GIsEditor)
+	{
+		FColor HitProxyColor(ForceInit);
+		int32 NumRealInstances = InComponent->PerInstanceSMData.Num();
+		bool bUseRemapTable = InComponent->PerInstanceSMData.Num() == InComponent->InstanceReorderTable.Num();
+		
+		for (int32 InstanceIndex = 0; InstanceIndex < NumRealInstances; InstanceIndex++)
+		{
+			if (InHitProxies.Num() == NumRealInstances)
+			{
+				HitProxyColor = InHitProxies[InstanceIndex]->Id.GetColor();
+			}
+			// Record if the instance is selected
+			bool bSelected = InstanceIndex < InComponent->SelectedInstances.Num() && InComponent->SelectedInstances[InstanceIndex];
+			FInstanceStream* InstanceRenderData = InstanceData->GetInstanceWriteAddress(bUseRemapTable ? InComponent->InstanceReorderTable[InstanceIndex] : InstanceIndex);
+			InstanceRenderData->SetInstanceEditorData(HitProxyColor, bSelected);
+		}
+	}
+#endif
+}
+
 /**
  * Specialized assignment operator, only used when importing LOD's.  
  */
@@ -710,6 +725,7 @@ void FInstancedStaticMeshSceneProxy::GetDistanceFieldInstanceInfo(int32& NumInst
 UInstancedStaticMeshComponent::UInstancedStaticMeshComponent(const FObjectInitializer& ObjectInitializer)
 	: Super(ObjectInitializer)
 	, bPerInstanceRenderDataWasPrebuilt(false)
+	, SelectionStamp(0)
 {
 	Mobility = EComponentMobility::Movable;
 	BodyInstance.bSimulatePhysics = false;
@@ -1267,6 +1283,8 @@ void UInstancedStaticMeshComponent::ReleasePerInstanceRenderData()
 	{
 		typedef TSharedPtr<FPerInstanceRenderData, ESPMode::ThreadSafe> FPerInstanceRenderDataPtr;
 
+		PerInstanceRenderData->HitProxies.Empty();
+
 		// Make shared pointer object on the heap
 		FPerInstanceRenderDataPtr* CleanupRenderDataPtr = new FPerInstanceRenderDataPtr(PerInstanceRenderData);
 		PerInstanceRenderData.Reset();
@@ -1664,7 +1682,15 @@ void UInstancedStaticMeshComponent::SelectInstance(bool bInSelected, int32 InIns
 		SelectedInstances[InstanceIndex] = bInSelected;
 	}
 
-	ReleasePerInstanceRenderData();
+	SelectionStamp++;
+#endif
+}
+
+void UInstancedStaticMeshComponent::ClearInstanceSelection()
+{
+#if WITH_EDITOR	
+	SelectedInstances.Empty();
+	SelectionStamp++;
 #endif
 }
 
diff --git a/Engine/Source/Runtime/Engine/Private/InstancedStaticMesh.h b/Engine/Source/Runtime/Engine/Private/InstancedStaticMesh.h
--- a/Engine/Source/Runtime/Engine/Private/InstancedStaticMesh.h
+++ b/Engine/Source/Runtime/Engine/Private/InstancedStaticMesh.h
@@ -94,6 +94,9 @@ public:
 	 */
 	void InitFromPreallocatedData(UInstancedStaticMeshComponent* InComponent, FStaticMeshInstanceData& Other);
 
+	/** Propagates instance selection state and hit proxy colors */
+	void SetPerInstanceEditorData(UInstancedStaticMeshComponent* InComponent, const TArray<TRefCountPtr<HHitProxy>>& InHitProxies);
+
 	/** Serializer. */
 	friend FArchive& operator<<(FArchive& Ar, FStaticMeshInstanceBuffer& VertexBuffer);
 
@@ -309,6 +312,7 @@ struct FPerInstanceRenderData
 	// Should be always constructed on main thread
 	FPerInstanceRenderData(UInstancedStaticMeshComponent* InComponent, ERHIFeatureLevel::Type InFeaureLevel)
 		: InstanceBuffer(InFeaureLevel)
+		, CachedSelectionStamp(InComponent->SelectionStamp)
 	{
 		// Create hit proxies for each instance if the component wants
 		if (GIsEditor && InComponent->bHasPerInstanceHitProxies)
@@ -328,6 +332,7 @@ struct FPerInstanceRenderData
 
 	FPerInstanceRenderData(UInstancedStaticMeshComponent* InComponent, FStaticMeshInstanceData& Other, ERHIFeatureLevel::Type InFeaureLevel)
 		: InstanceBuffer(InFeaureLevel)
+		, CachedSelectionStamp(0)
 	{
 		InstanceBuffer.InitFromPreallocatedData(InComponent, Other);
 		BeginInitResource(&InstanceBuffer);
@@ -343,6 +348,8 @@ struct FPerInstanceRenderData
 	FStaticMeshInstanceBuffer			InstanceBuffer;
 	/** Hit proxies for the instances */
 	TArray<TRefCountPtr<HHitProxy>>		HitProxies;
+	/** Cached selection identifier from component */
+	int32								CachedSelectionStamp;
 };
 
 
@@ -370,6 +377,13 @@ public:
 			PerInstanceRenderData = InComponent->PerInstanceRenderData;
 			InComponent->bPerInstanceRenderDataWasPrebuilt = false;
 		}
+		else if (PerInstanceRenderData->CachedSelectionStamp != InComponent->SelectionStamp)
+		{
+			// Update editor data bits in already existing instance buffer
+			PerInstanceRenderData->InstanceBuffer.SetPerInstanceEditorData(InComponent, PerInstanceRenderData->HitProxies);
+			PerInstanceRenderData->CachedSelectionStamp = InComponent->SelectionStamp;
+			BeginUpdateResourceRHI(&PerInstanceRenderData->InstanceBuffer);
+		}
 
 		NumInstances = PerInstanceRenderData->InstanceBuffer.GetNumInstances();
 		InitResources();
diff --git a/Engine/Source/Runtime/Engine/Public/StaticMeshResources.h b/Engine/Source/Runtime/Engine/Public/StaticMeshResources.h
--- a/Engine/Source/Runtime/Engine/Public/StaticMeshResources.h
+++ b/Engine/Source/Runtime/Engine/Public/StaticMeshResources.h
@@ -1045,6 +1045,15 @@ struct FInstanceStream
 		Me->InstanceTransform3[2] = FFloat16();
 	}
 
+	FORCEINLINE void SetInstanceEditorData(FColor HitProxyColor, bool bSelected)
+	{
+		FInstanceStream* RESTRICT Me = (FInstanceStream* RESTRICT)this;
+		
+		Me->InstanceTransform1[3] = ((float)HitProxyColor.R) + (bSelected ? 256.f : 0.0f);
+		Me->InstanceTransform2[3] = (float)HitProxyColor.G;
+		Me->InstanceTransform3[3] = (float)HitProxyColor.B;
+	}
+	
 	friend FArchive& operator<<( FArchive& Ar, FInstanceStream& V )
 	{
 		return Ar 
diff --git a/Engine/Source/Runtime/Foliage/Private/InstancedFoliage.cpp b/Engine/Source/Runtime/Foliage/Private/InstancedFoliage.cpp
--- a/Engine/Source/Runtime/Foliage/Private/InstancedFoliage.cpp
+++ b/Engine/Source/Runtime/Foliage/Private/InstancedFoliage.cpp
@@ -944,13 +944,11 @@ void FFoliageMeshInfo::ReapplyInstancesToComponent()
 
 		if (SelectedIndices.Num())
 		{
-			if (Component->SelectedInstances.Num() != Component->PerInstanceSMData.Num())
-			{
-				Component->SelectedInstances.Init(false, Component->PerInstanceSMData.Num());
-			}
+			Component->ClearInstanceSelection();
+			
 			for (int32 i : SelectedIndices)
 			{
-				Component->SelectedInstances[i] = true;
+				Component->SelectInstance(true, i, 1);
 			}
 		}
 
@@ -1035,15 +1033,15 @@ void FFoliageMeshInfo::SelectInstances(AInstancedFoliageActor* InIFA, bool bSele
 			{
 				SelectedIndices.Add(i);
 			}
+			
+			Component->SelectInstance(true, 0, Component->PerInstanceSMData.Num());
 		}
 		else
 		{
 			SelectedIndices.Empty();
+			Component->ClearInstanceSelection();
 		}
 
-		// Apply selections to the component
-		Component->SelectedInstances.Init(bSelect, Component->PerInstanceSMData.Num());
-		Component->ReleasePerInstanceRenderData();
 		Component->MarkRenderStateDirty();
 	}
 }
@@ -1053,37 +1051,35 @@ void FFoliageMeshInfo::SelectInstances(AInstancedFoliageActor* InIFA, bool bSele
 	if (InInstances.Num())
 	{
 		check(Component);
-		InIFA->Modify();
-
 		if (bSelect)
 		{
-			// Apply selections to the component
-			Component->ReleasePerInstanceRenderData();
-			Component->MarkRenderStateDirty();
-
-			if (Component->SelectedInstances.Num() != Component->PerInstanceSMData.Num())
-			{
-				Component->SelectedInstances.Init(false, Component->PerInstanceSMData.Num());
-			}
+			InIFA->Modify();
 
 			for (int32 i : InInstances)
 			{
 				SelectedIndices.Add(i);
-				Component->SelectedInstances[i] = true;
+				Component->SelectInstance(true, i, 1);
 			}
+			
+			Component->MarkRenderStateDirty();
 		}
 		else
 		{
-			if (Component->SelectedInstances.Num())
+			InIFA->Modify();
+
+			for (int32 i : InInstances)
 			{
-				Component->ReleasePerInstanceRenderData();
-				Component->MarkRenderStateDirty();
+				SelectedIndices.Remove(i);
+			}
 
+			if (Component->SelectedInstances.Num() > 0)
+			{
 				for (int32 i : InInstances)
 				{
-					SelectedIndices.Remove(i);
-					Component->SelectedInstances[i] = false;
+					Component->SelectInstance(false, i, 1);
 				}
+
+				Component->MarkRenderStateDirty();
 			}
 		}
 	}
@@ -1747,17 +1743,9 @@ void AInstancedFoliageActor::SelectInstance(UInstancedStaticMeshComponent* InCom
 		{
 			FFoliageMeshInfo& MeshInfo = *MeshPair.Value;
 
-			if (MeshInfo.SelectedIndices.Num() > 0)
-			{
-				check(MeshInfo.Component);
-				if (MeshInfo.Component->SelectedInstances.Num() > 0)
-				{
-					MeshInfo.Component->SelectedInstances.Empty();
-					MeshInfo.Component->ReleasePerInstanceRenderData();
-					MeshInfo.Component->MarkRenderStateDirty();
-				}
-			}
-
+			check(MeshInfo.Component);
+			MeshInfo.Component->ClearInstanceSelection();
+			MeshInfo.Component->MarkRenderStateDirty();
 			MeshInfo.SelectedIndices.Empty();
 		}
 	}
@@ -1773,6 +1761,7 @@ void AInstancedFoliageActor::SelectInstance(UInstancedStaticMeshComponent* InCom
 			{
 				Type = MeshPair.Key;
 				MeshInfo = &MeshPair.Value.Get();
+				break;
 			}
 		}
 		
@@ -1783,8 +1772,7 @@ void AInstancedFoliageActor::SelectInstance(UInstancedStaticMeshComponent* InCom
 			// Deselect if it's already selected.
 			if (InInstanceIndex < InComponent->SelectedInstances.Num())
 			{
-				InComponent->SelectedInstances[InInstanceIndex] = false;
-				InComponent->ReleasePerInstanceRenderData();
+				InComponent->SelectInstance(false, InInstanceIndex, 1);
 				InComponent->MarkRenderStateDirty();
 			}
 
@@ -1796,12 +1784,7 @@ void AInstancedFoliageActor::SelectInstance(UInstancedStaticMeshComponent* InCom
 			if (!bToggle || !bIsSelected)
 			{
 				// Add the selection
-				if (InComponent->SelectedInstances.Num() < InComponent->PerInstanceSMData.Num())
-				{
-					InComponent->SelectedInstances.Init(false, InComponent->PerInstanceSMData.Num());
-				}
-				InComponent->SelectedInstances[InInstanceIndex] = true;
-				InComponent->ReleasePerInstanceRenderData();
+				InComponent->SelectInstance(true, InInstanceIndex, 1);
 				InComponent->MarkRenderStateDirty();
 
 				MeshInfo->SelectedIndices.Add(InInstanceIndex);
@@ -1878,34 +1861,21 @@ void AInstancedFoliageActor::ApplySelectionToComponents(bool bApply)
 	for (auto& MeshPair : FoliageMeshes)
 	{
 		FFoliageMeshInfo& MeshInfo = *MeshPair.Value;
+		UHierarchicalInstancedStaticMeshComponent* Component = MeshInfo.Component;
 
-		if (bApply)
+		if (Component && (bApply || Component->SelectedInstances.Num() > 0))
 		{
-			if (MeshInfo.SelectedIndices.Num() > 0)
+			Component->ClearInstanceSelection();
+			
+			if (bApply)
 			{
-				check(MeshInfo.Component);
-
-				// Apply any selections in the component
-				MeshInfo.Component->SelectedInstances.Init(false, MeshInfo.Component->PerInstanceSMData.Num());
 				for (int32 i : MeshInfo.SelectedIndices)
 				{
-					MeshInfo.Component->SelectedInstances[i] = true;
+					Component->SelectInstance(true, i, 1);
 				}
-
-				MeshInfo.Component->ReleasePerInstanceRenderData();
-				MeshInfo.Component->MarkRenderStateDirty();
 			}
-		}
-		else		
-		{
-			if (MeshInfo.Component && MeshInfo.Component->SelectedInstances.Num() > 0)
-			{
-				// remove any selections in the component
-				MeshInfo.Component->SelectedInstances.Empty();
 
-				MeshInfo.Component->ReleasePerInstanceRenderData();
-				MeshInfo.Component->MarkRenderStateDirty();
-			}
+			Component->MarkRenderStateDirty();
 		}
 	}
 }
