Fixed error in procedural foliage simulation in which the initial seed location(s) on every tile were the same

diff --git a/Engine/Source/Runtime/Foliage/Private/ProceduralFoliageTile.cpp b/Engine/Source/Runtime/Foliage/Private/ProceduralFoliageTile.cpp
--- a/Engine/Source/Runtime/Foliage/Private/ProceduralFoliageTile.cpp
+++ b/Engine/Source/Runtime/Foliage/Private/ProceduralFoliageTile.cpp
@@ -222,7 +222,7 @@ void UProceduralFoliageTile::AddRandomSeeds(TArray<FProceduralFoliageInstance*>&
 			}
 
 			{	//save the random stream per type
-				RandomStreamPerType.Add(Type, FRandomStream(Type->DistributionSeed + FoliageSpawner->RandomSeed));
+				RandomStreamPerType.Add(Type, FRandomStream(Type->DistributionSeed + FoliageSpawner->RandomSeed + RandomSeed));
 			}
 
 			{	//compute the needed offsets for initial seed variance
@@ -370,8 +370,9 @@ void UProceduralFoliageTile::FlushPendingRemovals()
 	PendingRemovals.Empty();
 }
 
-void UProceduralFoliageTile::InitSimulation(const UProceduralFoliageSpawner* InFoliageSpawner, const int32 RandomSeed)
+void UProceduralFoliageTile::InitSimulation(const UProceduralFoliageSpawner* InFoliageSpawner, const int32 InRandomSeed)
 {
+	RandomSeed = InRandomSeed;
 	RandomStream.Initialize(RandomSeed);
 	FoliageSpawner = InFoliageSpawner;
 	SimulationStep = 0;
diff --git a/Engine/Source/Runtime/Foliage/Public/ProceduralFoliageTile.h b/Engine/Source/Runtime/Foliage/Public/ProceduralFoliageTile.h
--- a/Engine/Source/Runtime/Foliage/Public/ProceduralFoliageTile.h
+++ b/Engine/Source/Runtime/Foliage/Public/ProceduralFoliageTile.h
@@ -60,7 +60,7 @@ class FOLIAGE_API UProceduralFoliageTile : public UObject
 	void RemoveInstances();
 
 	/** Initializes the tile for simulation. Called automatically by Simulate, so only bother to call if the tile won't actually simulate. */
-	void InitSimulation(const UProceduralFoliageSpawner* InFoliageSpawner, const int32 RandomSeed);
+	void InitSimulation(const UProceduralFoliageSpawner* InFoliageSpawner, const int32 InRandomSeed);
 
 	// UObject interface
 	virtual SIZE_T GetResourceSize(EResourceSizeMode::Type Mode) override;
@@ -131,6 +131,7 @@ private:
 	int32 SimulationStep;
 	FProceduralFoliageBroadphase Broadphase;	
 
+	int32 RandomSeed;
 	FRandomStream RandomStream;
 	bool bSimulateShadeGrowth;
 
