[UETOOL-269] - Foliage type assets now work with procedural foliage in addition to BP subclasses

- Renamed FProceduralFoliageTypeData to FFoliageTypeObject
- Added property customization for FFoliageTypeObject
- Can contain either a FoliageType asset or BP subclass

diff --git a/Engine/Config/BaseEngine.ini b/Engine/Config/BaseEngine.ini
--- a/Engine/Config/BaseEngine.ini
+++ b/Engine/Config/BaseEngine.ini
@@ -360,6 +360,8 @@ bUseStreamingPause=false
 +ActiveClassRedirects=(OldClassName="ProceduralFoliage",NewClassName="ProceduralFoliageSpawner")
 +TaggedPropertyRedirects=(ClassName="ProceduralFoliageComponent",OldPropertyName="ProceduralFoliage",NewPropertyName="FoliageSpawner")
 +TaggedPropertyRedirects=(ClassName="ProceduralFoliageComponent",OldPropertyName="Overlap",NewPropertyName="TileOverlap")
++ActiveStructRedirects=(OldStructName="ProceduralFoliageTypeData",NewStructName="FoliageTypeObject")
++TaggedPropertyRedirects=(ClassName="ProceduralFoliageSpawner", OldPropertyName="Types", NewPropertyName="FoliageTypes")
 
 ;Shortened animation compression algorithm class names
 +ActiveClassRedirects=(OldClassName="AnimationCompressionAlgorithm",NewClassName="AnimCompress")
diff --git a/Engine/Source/Editor/FoliageEdit/Private/FoliageEditModule.cpp b/Engine/Source/Editor/FoliageEdit/Private/FoliageEditModule.cpp
--- a/Engine/Source/Editor/FoliageEdit/Private/FoliageEditModule.cpp
+++ b/Engine/Source/Editor/FoliageEdit/Private/FoliageEditModule.cpp
@@ -16,6 +16,7 @@ const FName FoliageEditAppIdentifier = FName(TEXT("FoliageEdApp"));
 #include "ComponentVisualizer.h"
 #include "ProceduralFoliageVolume.h"
 #include "ProceduralFoliageComponent.h"
+#include "FoliageTypeObjectCustomization.h"
 
 /**
  * Foliage Edit Mode module
@@ -39,7 +40,8 @@ public:
 		// Register the details customizer
 		FPropertyEditorModule& PropertyModule = FModuleManager::LoadModuleChecked<FPropertyEditorModule>("PropertyEditor");
 		PropertyModule.RegisterCustomClassLayout("FoliageType", FOnGetDetailCustomizationInstance::CreateStatic(&FFoliageTypeDetails::MakeInstance));
-
+		PropertyModule.RegisterCustomPropertyTypeLayout("FoliageTypeObject", FOnGetPropertyTypeCustomizationInstance::CreateStatic(&FFoliageTypeObjectCustomization::MakeInstance));
+		
 		GUnrealEd->RegisterComponentVisualizer(UProceduralFoliageComponent::StaticClass()->GetFName(), MakeShareable(new FProceduralFoliageComponentVisualizer));
 
 		FPropertyEditorModule& PropertyEditor = FModuleManager::LoadModuleChecked<FPropertyEditorModule>("PropertyEditor");
diff --git a/Engine/Source/Editor/PropertyEditor/Public/PropertyCustomizationHelpers.h b/Engine/Source/Editor/PropertyEditor/Public/PropertyCustomizationHelpers.h
--- a/Engine/Source/Editor/PropertyEditor/Public/PropertyCustomizationHelpers.h
+++ b/Engine/Source/Editor/PropertyEditor/Public/PropertyCustomizationHelpers.h
@@ -70,7 +70,7 @@ public:
 		SLATE_ARGUMENT( TSharedPtr<IPropertyHandle>, PropertyHandle )
 		/** Thumbnail pool */
 		SLATE_ARGUMENT( TSharedPtr<FAssetThumbnailPool>, ThumbnailPool )
-		/** Classes that are allowed in the asset picker */
+		/** Class that is allowed in the asset picker */
 		SLATE_ARGUMENT( UClass*, AllowedClass )
 		/** Optional list of factories which may be used to create new assets */
 		SLATE_ARGUMENT( TOptional<TArray<UFactory*>>, NewAssetFactories )
diff --git a/Engine/Source/Runtime/Core/Public/UObject/ObjectVersion.h b/Engine/Source/Runtime/Core/Public/UObject/ObjectVersion.h
--- a/Engine/Source/Runtime/Core/Public/UObject/ObjectVersion.h
+++ b/Engine/Source/Runtime/Core/Public/UObject/ObjectVersion.h
@@ -490,6 +490,8 @@ enum EUnrealEngineObjectUE4Version
 	VER_UE4_LANDSCAPE_SPLINE_CROSS_LEVEL_MESHES,
 	// Deprecate the variables used for sizing in the designer on UUserWidget
 	VER_UE4_DEPRECATE_USER_WIDGET_DESIGN_SIZE,
+	// Updated foliage to work with either FoliageType assets or blueprint classes
+	VER_UE4_FOLIAGE_WITH_ASSET_OR_CLASS,
 
 	// -----<new versions can be added before this line>-------------------------------------------------
 	// - this needs to be the last line (see note below)
diff --git a/Engine/Source/Runtime/Foliage/Private/FoliageTypeObject.cpp b/Engine/Source/Runtime/Foliage/Private/FoliageTypeObject.cpp
new file mode 100644
--- /dev/null
+++ b/Engine/Source/Runtime/Foliage/Private/FoliageTypeObject.cpp
@@ -0,0 +1,67 @@
+// Copyright 1998-2015 Epic Games, Inc. All Rights Reserved.
+
+#pragma once
+#include "FoliagePrivate.h"
+#include "FoliageTypeObject.h"
+
+void FFoliageTypeObject::RefreshInstance()
+{
+	//@todo: should verify that the instance is dirty before updating
+
+	// If the foliage type is an asset, use that
+	if (auto FoliageTypeAsAsset = Cast<UFoliageType_InstancedStaticMesh>(FoliageTypeObject))
+	{
+		TypeInstance = FoliageTypeAsAsset;
+	}
+	else if (auto FoliageTypeAsBP = Cast<UBlueprint>(FoliageTypeObject))
+	{
+		// If the foliage type is a BP class, create and store an instance of that class
+		if (FoliageTypeAsBP->ParentClass->IsChildOf<UFoliageType_InstancedStaticMesh>())
+		{
+			TypeInstance = NewObject<UFoliageType_InstancedStaticMesh>(GetTransientPackage(), *FoliageTypeAsBP->GeneratedClass);
+		}
+	}
+}
+
+const UFoliageType_InstancedStaticMesh* FFoliageTypeObject::GetInstance()
+{
+	if (!TypeInstance)
+	{
+		// If the type instance doesn't exist, try to create it
+		RefreshInstance();
+	}
+
+	return TypeInstance;
+}
+
+const UFoliageType_InstancedStaticMesh* FFoliageTypeObject::GetInstance() const
+{
+	return TypeInstance;
+}
+
+bool FFoliageTypeObject::IsValid() const
+{
+	return TypeInstance != nullptr;
+}
+
+bool FFoliageTypeObject::IsDirty() const
+{
+	return (!TypeInstance || TypeInstance->ChangeCount != TypeInstance->GetClass()->GetDefaultObject<UFoliageType_InstancedStaticMesh>()->ChangeCount);
+}
+
+void FFoliageTypeObject::SetClean()
+{
+	if (TypeInstance)
+	{
+		TypeInstance->ChangeCount = TypeInstance->GetClass()->GetDefaultObject<UFoliageType_InstancedStaticMesh>()->ChangeCount;
+	}
+}
+
+void FFoliageTypeObject::PostSerialize(const FArchive& Ar)
+{
+	if (Ar.UE4Ver() < VER_UE4_FOLIAGE_WITH_ASSET_OR_CLASS)
+	{
+		// Set the type object to be the previously set FoliageType BP class
+		FoliageTypeObject = Type_DEPRECATED->ClassGeneratedBy;
+	}
+}
diff --git a/Engine/Source/Runtime/Foliage/Private/ProceduralFoliageSpawner.cpp b/Engine/Source/Runtime/Foliage/Private/ProceduralFoliageSpawner.cpp
--- a/Engine/Source/Runtime/Foliage/Private/ProceduralFoliageSpawner.cpp
+++ b/Engine/Source/Runtime/Foliage/Private/ProceduralFoliageSpawner.cpp
@@ -34,28 +34,18 @@ UProceduralFoliageTile* UProceduralFoliageSpawner::CreateTempTile()
 
 void UProceduralFoliageSpawner::CreateProceduralFoliageInstances()
 {
-	for(FProceduralFoliageTypeData& TypeData : Types)
+	for (FFoliageTypeObject& FoliageTypeObject : FoliageTypes)
 	{
-		if( TypeData.Type )
-		{
-			TypeData.TypeInstance = NewObject<UFoliageType_InstancedStaticMesh>(this, TypeData.Type);
-			TypeData.ChangeCount = TypeData.TypeInstance->ChangeCount;
-		}
-		else
-		{
-			TypeData.TypeInstance = nullptr;
-		}
+		// Refresh the instances contained in the type objects
+		FoliageTypeObject.RefreshInstance();
 	}
 }
 
 void UProceduralFoliageSpawner::SetClean()
 {
-	for (FProceduralFoliageTypeData& TypeData : Types)
+	for (FFoliageTypeObject& FoliageTypeObject : FoliageTypes)
 	{
-		if( TypeData.Type && TypeData.TypeInstance )
-		{
-			TypeData.TypeInstance->ChangeCount = TypeData.Type->GetDefaultObject<UFoliageType_InstancedStaticMesh>()->ChangeCount;
-		}
+		FoliageTypeObject.SetClean();
 	}
 
 	bNeedsSimulation = false;
@@ -97,23 +87,18 @@ bool UProceduralFoliageSpawner::AnyDirty() const
 {
 	bool bDirty = bNeedsSimulation;
 
-	for (const FProceduralFoliageTypeData& TypeData : Types)
+	if (!bDirty)
 	{
-		if (TypeData.Type)
+		for (const FFoliageTypeObject& FoliageTypeObject : FoliageTypes)
 		{
-			if (TypeData.TypeInstance && TypeData.Type->GetDefaultObject<UFoliageType_InstancedStaticMesh>()->ChangeCount != TypeData.ChangeCount)
-			{
-				bDirty = true;
-				break;
-			}
-			else if (!TypeData.TypeInstance)
+			if (FoliageTypeObject.IsDirty())
 			{
 				bDirty = true;
 				break;
 			}
 		}
 	}
-
+	
 	return bDirty;
 }
 
diff --git a/Engine/Source/Runtime/Foliage/Private/ProceduralFoliageTile.cpp b/Engine/Source/Runtime/Foliage/Private/ProceduralFoliageTile.cpp
--- a/Engine/Source/Runtime/Foliage/Private/ProceduralFoliageTile.cpp
+++ b/Engine/Source/Runtime/Foliage/Private/ProceduralFoliageTile.cpp
@@ -206,29 +206,29 @@ void UProceduralFoliageTile::AddRandomSeeds(TArray<FProceduralFoliageInstance*>&
 
 	TArray<const UFoliageType_InstancedStaticMesh*> TypesToSeed;
 
-	for (const FProceduralFoliageTypeData& Data : FoliageSpawner->GetTypes())
+	for (const FFoliageTypeObject& FoliageTypeObject : FoliageSpawner->GetFoliageTypes())
 	{
 		if (UserCancelled()){ return; }
-		const UFoliageType_InstancedStaticMesh* Type = Data.TypeInstance;
-		if (Type && Type->bGrowsInShade == bSimulateShadeGrowth)
+		const UFoliageType_InstancedStaticMesh* TypeInstance = FoliageTypeObject.GetInstance();
+		if (TypeInstance && TypeInstance->bGrowsInShade == bSimulateShadeGrowth)
 		{
 			{	//compute the number of initial seeds
-				const int32 NumSeeds = FMath::RoundToInt(Type->GetSeedDensitySquared() * SizeTenM2);
-				SeedsLeftMap.Add(Type, NumSeeds);
+				const int32 NumSeeds = FMath::RoundToInt(TypeInstance->GetSeedDensitySquared() * SizeTenM2);
+				SeedsLeftMap.Add(TypeInstance, NumSeeds);
 				if (NumSeeds > 0)
 				{
-					TypesToSeed.Add(Type);
+					TypesToSeed.Add(TypeInstance);
 				}
 			}
 
 			{	//save the random stream per type
-				RandomStreamPerType.Add(Type, FRandomStream(Type->DistributionSeed + FoliageSpawner->RandomSeed + RandomSeed));
+				RandomStreamPerType.Add(TypeInstance, FRandomStream(TypeInstance->DistributionSeed + FoliageSpawner->RandomSeed + RandomSeed));
 			}
 
 			{	//compute the needed offsets for initial seed variance
-				const int32 DistributionSeed = Type->DistributionSeed;
-				const float MaxScale = Type->GetScaleForAge(Type->MaxAge);
-				const float TypeMaxCollisionRadius = MaxScale * Type->CollisionRadius;
+				const int32 DistributionSeed = TypeInstance->DistributionSeed;
+				const float MaxScale = TypeInstance->GetScaleForAge(TypeInstance->MaxAge);
+				const float TypeMaxCollisionRadius = MaxScale * TypeInstance->CollisionRadius;
 				if (float* MaxRadius = MaxCollisionRadii.Find(DistributionSeed))
 				{
 					*MaxRadius = FMath::Max(*MaxRadius, TypeMaxCollisionRadius);
@@ -238,7 +238,7 @@ void UProceduralFoliageTile::AddRandomSeeds(TArray<FProceduralFoliageInstance*>&
 					MaxCollisionRadii.Add(DistributionSeed, TypeMaxCollisionRadius);
 				}
 
-				const float TypeMaxShadeRadius = MaxScale * Type->ShadeRadius;
+				const float TypeMaxShadeRadius = MaxScale * TypeInstance->ShadeRadius;
 				if (float* MaxRadius = MaxShadeRadii.Find(DistributionSeed))
 				{
 					*MaxRadius = FMath::Max(*MaxRadius, TypeMaxShadeRadius);
@@ -411,9 +411,9 @@ void UProceduralFoliageTile::RunSimulation(const int32 MaxNumSteps, bool bShadeG
 {
 	int32 MaxSteps = 0;
 
-	for (const FProceduralFoliageTypeData& Data : FoliageSpawner->GetTypes())
+	for (const FFoliageTypeObject& FoliageTypeObject : FoliageSpawner->GetFoliageTypes())
 	{
-		const UFoliageType_InstancedStaticMesh* TypeInstance = Data.TypeInstance;
+		const UFoliageType_InstancedStaticMesh* TypeInstance = FoliageTypeObject.GetInstance();
 		if (TypeInstance && TypeInstance->bGrowsInShade == bShadeGrowth)
 		{
 			MaxSteps = FMath::Max(MaxSteps, TypeInstance->NumSteps + 1);
diff --git a/Engine/Source/Runtime/Foliage/Public/FoliageTypeObject.h b/Engine/Source/Runtime/Foliage/Public/FoliageTypeObject.h
new file mode 100644
--- /dev/null
+++ b/Engine/Source/Runtime/Foliage/Public/FoliageTypeObject.h
@@ -0,0 +1,54 @@
+// Copyright 1998-2015 Epic Games, Inc. All Rights Reserved.
+
+#pragma once
+
+#include "FoliageType_InstancedStaticMesh.h"
+#include "FoliageTypeObject.generated.h"
+
+/** A wrapper struct used to allow the use of either FoliageType assets or FoliageType blueprint classes */
+USTRUCT()
+struct FFoliageTypeObject
+{
+	GENERATED_USTRUCT_BODY()
+
+	FFoliageTypeObject()
+		: FoliageTypeObject(nullptr), TypeInstance(nullptr) {}
+
+	/** Refreshes the type instance based on the assigned object. Intended to be called after some change is made. */
+	void RefreshInstance();
+
+	/** Gets the instance of this foliage type. Creates the instance if needed. */
+	const UFoliageType_InstancedStaticMesh* GetInstance();
+
+	/** Gets the instance of this foliage type. */
+	const UFoliageType_InstancedStaticMesh* GetInstance() const;
+
+	/** @return Whether this would return a valid instance */
+	bool IsValid() const;
+
+	bool IsDirty() const;
+	void SetClean();
+
+	void PostSerialize(const FArchive& Ar);
+
+private:
+	/** The foliage type that will be spawned by the procedural foliage simulation */
+	UPROPERTY(Category = ProceduralFoliageSimulation, EditAnywhere, meta=(AllowedClasses="FoliageType_InstancedStaticMesh,Blueprint", DisplayThumbnail="true", ThumbnailSize="X=40 Y=40"))
+	UObject* FoliageTypeObject;
+
+	/** The actual instance of the foliage type that is used for spawning */
+	UPROPERTY(transient)
+	UFoliageType_InstancedStaticMesh* TypeInstance;
+
+	UPROPERTY()
+	TSubclassOf<UFoliageType_InstancedStaticMesh> Type_DEPRECATED;
+};
+
+template<>
+struct TStructOpsTypeTraits<FFoliageTypeObject> : public TStructOpsTypeTraitsBase
+{
+	enum
+	{
+		WithPostSerialize = true,
+	};
+};
diff --git a/Engine/Source/Runtime/Foliage/Public/FoliageType_InstancedStaticMesh.h b/Engine/Source/Runtime/Foliage/Public/FoliageType_InstancedStaticMesh.h
--- a/Engine/Source/Runtime/Foliage/Public/FoliageType_InstancedStaticMesh.h
+++ b/Engine/Source/Runtime/Foliage/Public/FoliageType_InstancedStaticMesh.h
@@ -10,7 +10,7 @@ class UFoliageType_InstancedStaticMesh : public UFoliageType
 {
 	GENERATED_UCLASS_BODY()
 
-	UPROPERTY(EditAnywhere, Category=FoliageType)
+	UPROPERTY(EditAnywhere, Category=Mesh)
 	UStaticMesh* Mesh;
 
 	virtual UStaticMesh* GetStaticMesh() const override
diff --git a/Engine/Source/Runtime/Foliage/Public/ProceduralFoliageSpawner.h b/Engine/Source/Runtime/Foliage/Public/ProceduralFoliageSpawner.h
--- a/Engine/Source/Runtime/Foliage/Public/ProceduralFoliageSpawner.h
+++ b/Engine/Source/Runtime/Foliage/Public/ProceduralFoliageSpawner.h
@@ -4,26 +4,11 @@
 
 #include "FoliageType_InstancedStaticMesh.h"
 #include "ProceduralFoliageInstance.h"
+#include "FoliageTypeObject.h"
 #include "ProceduralFoliageSpawner.generated.h"
 
 class UProceduralFoliageTile;
 
-
-USTRUCT()
-struct FProceduralFoliageTypeData
-{
-	GENERATED_USTRUCT_BODY()
-
-	UPROPERTY(Category = ProceduralFoliageSimulation, EditAnywhere, BlueprintReadOnly)
-	TSubclassOf<UFoliageType_InstancedStaticMesh> Type;
-
-	UPROPERTY(transient)
-	UFoliageType_InstancedStaticMesh* TypeInstance;
-
-	UPROPERTY()
-	int32 ChangeCount;
-};
-
 UCLASS(BlueprintType, Blueprintable)
 class FOLIAGE_API UProceduralFoliageSpawner : public UObject
 {
@@ -46,7 +31,7 @@ class FOLIAGE_API UProceduralFoliageSpawner : public UObject
 private:
 	/** The types of foliage to procedurally spawn. */
 	UPROPERTY(Category = ProceduralFoliageSimulation, EditAnywhere)
-	TArray<FProceduralFoliageTypeData> Types;
+	TArray<FFoliageTypeObject> FoliageTypes;
 
 	UPROPERTY()
 	bool bNeedsSimulation;
@@ -57,7 +42,7 @@ public:
 
 	int32 GetRandomNumber();
 
-	const TArray<FProceduralFoliageTypeData>& GetTypes() const { return Types; }
+	const TArray<FFoliageTypeObject>& GetFoliageTypes() const { return FoliageTypes; }
 
 	/** Returns the instances that need to spawn for a given min,max region */
 	void GetInstancesToSpawn(TArray<FProceduralFoliageInstance>& OutInstances, const FVector& Min = FVector(-FLT_MAX, -FLT_MAX, -FLT_MAX), const FVector& Max = FVector(FLT_MAX, FLT_MAX, FLT_MAX) ) const;
