Back out changelist 2508810 - there was a missing file

diff --git a/Engine/Config/BaseEngine.ini b/Engine/Config/BaseEngine.ini
--- a/Engine/Config/BaseEngine.ini
+++ b/Engine/Config/BaseEngine.ini
@@ -360,8 +360,6 @@ bUseStreamingPause=false
 +ActiveClassRedirects=(OldClassName="ProceduralFoliage",NewClassName="ProceduralFoliageSpawner")
 +TaggedPropertyRedirects=(ClassName="ProceduralFoliageComponent",OldPropertyName="ProceduralFoliage",NewPropertyName="FoliageSpawner")
 +TaggedPropertyRedirects=(ClassName="ProceduralFoliageComponent",OldPropertyName="Overlap",NewPropertyName="TileOverlap")
-+ActiveStructRedirects=(OldStructName="ProceduralFoliageTypeData",NewStructName="FoliageTypeObject")
-+TaggedPropertyRedirects=(ClassName="ProceduralFoliageSpawner", OldPropertyName="Types", NewPropertyName="FoliageTypes")
 
 ;Shortened animation compression algorithm class names
 +ActiveClassRedirects=(OldClassName="AnimationCompressionAlgorithm",NewClassName="AnimCompress")
diff --git a/Engine/Source/Editor/FoliageEdit/Private/FoliageEditModule.cpp b/Engine/Source/Editor/FoliageEdit/Private/FoliageEditModule.cpp
--- a/Engine/Source/Editor/FoliageEdit/Private/FoliageEditModule.cpp
+++ b/Engine/Source/Editor/FoliageEdit/Private/FoliageEditModule.cpp
@@ -16,7 +16,6 @@ const FName FoliageEditAppIdentifier = FName(TEXT("FoliageEdApp"));
 #include "ComponentVisualizer.h"
 #include "ProceduralFoliageVolume.h"
 #include "ProceduralFoliageComponent.h"
-#include "FoliageTypeObjectCustomization.h"
 
 /**
  * Foliage Edit Mode module
@@ -40,8 +39,7 @@ public:
 		// Register the details customizer
 		FPropertyEditorModule& PropertyModule = FModuleManager::LoadModuleChecked<FPropertyEditorModule>("PropertyEditor");
 		PropertyModule.RegisterCustomClassLayout("FoliageType", FOnGetDetailCustomizationInstance::CreateStatic(&FFoliageTypeDetails::MakeInstance));
-		PropertyModule.RegisterCustomPropertyTypeLayout("FoliageTypeObject", FOnGetPropertyTypeCustomizationInstance::CreateStatic(&FFoliageTypeObjectCustomization::MakeInstance));
-		
+
 		GUnrealEd->RegisterComponentVisualizer(UProceduralFoliageComponent::StaticClass()->GetFName(), MakeShareable(new FProceduralFoliageComponentVisualizer));
 
 		FPropertyEditorModule& PropertyEditor = FModuleManager::LoadModuleChecked<FPropertyEditorModule>("PropertyEditor");
diff --git a/Engine/Source/Editor/PropertyEditor/Public/PropertyCustomizationHelpers.h b/Engine/Source/Editor/PropertyEditor/Public/PropertyCustomizationHelpers.h
--- a/Engine/Source/Editor/PropertyEditor/Public/PropertyCustomizationHelpers.h
+++ b/Engine/Source/Editor/PropertyEditor/Public/PropertyCustomizationHelpers.h
@@ -70,7 +70,7 @@ public:
 		SLATE_ARGUMENT( TSharedPtr<IPropertyHandle>, PropertyHandle )
 		/** Thumbnail pool */
 		SLATE_ARGUMENT( TSharedPtr<FAssetThumbnailPool>, ThumbnailPool )
-		/** Class that is allowed in the asset picker */
+		/** Classes that are allowed in the asset picker */
 		SLATE_ARGUMENT( UClass*, AllowedClass )
 		/** Optional list of factories which may be used to create new assets */
 		SLATE_ARGUMENT( TOptional<TArray<UFactory*>>, NewAssetFactories )
diff --git a/Engine/Source/Runtime/Core/Public/UObject/ObjectVersion.h b/Engine/Source/Runtime/Core/Public/UObject/ObjectVersion.h
--- a/Engine/Source/Runtime/Core/Public/UObject/ObjectVersion.h
+++ b/Engine/Source/Runtime/Core/Public/UObject/ObjectVersion.h
@@ -490,8 +490,6 @@ enum EUnrealEngineObjectUE4Version
 	VER_UE4_LANDSCAPE_SPLINE_CROSS_LEVEL_MESHES,
 	// Deprecate the variables used for sizing in the designer on UUserWidget
 	VER_UE4_DEPRECATE_USER_WIDGET_DESIGN_SIZE,
-	// Updated foliage to work with either FoliageType assets or blueprint classes
-	VER_UE4_FOLIAGE_WITH_ASSET_OR_CLASS,
 
 	// -----<new versions can be added before this line>-------------------------------------------------
 	// - this needs to be the last line (see note below)
diff --git a/Engine/Source/Runtime/Foliage/Private/FoliageTypeObject.cpp b/Engine/Source/Runtime/Foliage/Private/FoliageTypeObject.cpp
deleted file mode 100644
diff --git a/Engine/Source/Runtime/Foliage/Private/ProceduralFoliageSpawner.cpp b/Engine/Source/Runtime/Foliage/Private/ProceduralFoliageSpawner.cpp
--- a/Engine/Source/Runtime/Foliage/Private/ProceduralFoliageSpawner.cpp
+++ b/Engine/Source/Runtime/Foliage/Private/ProceduralFoliageSpawner.cpp
@@ -34,18 +34,28 @@ UProceduralFoliageTile* UProceduralFoliageSpawner::CreateTempTile()
 
 void UProceduralFoliageSpawner::CreateProceduralFoliageInstances()
 {
-	for (FFoliageTypeObject& FoliageTypeObject : FoliageTypes)
+	for(FProceduralFoliageTypeData& TypeData : Types)
 	{
-		// Refresh the instances contained in the type objects
-		FoliageTypeObject.RefreshInstance();
+		if( TypeData.Type )
+		{
+			TypeData.TypeInstance = NewObject<UFoliageType_InstancedStaticMesh>(this, TypeData.Type);
+			TypeData.ChangeCount = TypeData.TypeInstance->ChangeCount;
+		}
+		else
+		{
+			TypeData.TypeInstance = nullptr;
+		}
 	}
 }
 
 void UProceduralFoliageSpawner::SetClean()
 {
-	for (FFoliageTypeObject& FoliageTypeObject : FoliageTypes)
+	for (FProceduralFoliageTypeData& TypeData : Types)
 	{
-		FoliageTypeObject.SetClean();
+		if( TypeData.Type && TypeData.TypeInstance )
+		{
+			TypeData.TypeInstance->ChangeCount = TypeData.Type->GetDefaultObject<UFoliageType_InstancedStaticMesh>()->ChangeCount;
+		}
 	}
 
 	bNeedsSimulation = false;
@@ -87,18 +97,23 @@ bool UProceduralFoliageSpawner::AnyDirty() const
 {
 	bool bDirty = bNeedsSimulation;
 
-	if (!bDirty)
+	for (const FProceduralFoliageTypeData& TypeData : Types)
 	{
-		for (const FFoliageTypeObject& FoliageTypeObject : FoliageTypes)
+		if (TypeData.Type)
 		{
-			if (FoliageTypeObject.IsDirty())
+			if (TypeData.TypeInstance && TypeData.Type->GetDefaultObject<UFoliageType_InstancedStaticMesh>()->ChangeCount != TypeData.ChangeCount)
+			{
+				bDirty = true;
+				break;
+			}
+			else if (!TypeData.TypeInstance)
 			{
 				bDirty = true;
 				break;
 			}
 		}
 	}
-	
+
 	return bDirty;
 }
 
diff --git a/Engine/Source/Runtime/Foliage/Private/ProceduralFoliageTile.cpp b/Engine/Source/Runtime/Foliage/Private/ProceduralFoliageTile.cpp
--- a/Engine/Source/Runtime/Foliage/Private/ProceduralFoliageTile.cpp
+++ b/Engine/Source/Runtime/Foliage/Private/ProceduralFoliageTile.cpp
@@ -206,29 +206,29 @@ void UProceduralFoliageTile::AddRandomSeeds(TArray<FProceduralFoliageInstance*>&
 
 	TArray<const UFoliageType_InstancedStaticMesh*> TypesToSeed;
 
-	for (const FFoliageTypeObject& FoliageTypeObject : FoliageSpawner->GetFoliageTypes())
+	for (const FProceduralFoliageTypeData& Data : FoliageSpawner->GetTypes())
 	{
 		if (UserCancelled()){ return; }
-		const UFoliageType_InstancedStaticMesh* TypeInstance = FoliageTypeObject.GetInstance();
-		if (TypeInstance && TypeInstance->bGrowsInShade == bSimulateShadeGrowth)
+		const UFoliageType_InstancedStaticMesh* Type = Data.TypeInstance;
+		if (Type && Type->bGrowsInShade == bSimulateShadeGrowth)
 		{
 			{	//compute the number of initial seeds
-				const int32 NumSeeds = FMath::RoundToInt(TypeInstance->GetSeedDensitySquared() * SizeTenM2);
-				SeedsLeftMap.Add(TypeInstance, NumSeeds);
+				const int32 NumSeeds = FMath::RoundToInt(Type->GetSeedDensitySquared() * SizeTenM2);
+				SeedsLeftMap.Add(Type, NumSeeds);
 				if (NumSeeds > 0)
 				{
-					TypesToSeed.Add(TypeInstance);
+					TypesToSeed.Add(Type);
 				}
 			}
 
 			{	//save the random stream per type
-				RandomStreamPerType.Add(TypeInstance, FRandomStream(TypeInstance->DistributionSeed + FoliageSpawner->RandomSeed + RandomSeed));
+				RandomStreamPerType.Add(Type, FRandomStream(Type->DistributionSeed + FoliageSpawner->RandomSeed + RandomSeed));
 			}
 
 			{	//compute the needed offsets for initial seed variance
-				const int32 DistributionSeed = TypeInstance->DistributionSeed;
-				const float MaxScale = TypeInstance->GetScaleForAge(TypeInstance->MaxAge);
-				const float TypeMaxCollisionRadius = MaxScale * TypeInstance->CollisionRadius;
+				const int32 DistributionSeed = Type->DistributionSeed;
+				const float MaxScale = Type->GetScaleForAge(Type->MaxAge);
+				const float TypeMaxCollisionRadius = MaxScale * Type->CollisionRadius;
 				if (float* MaxRadius = MaxCollisionRadii.Find(DistributionSeed))
 				{
 					*MaxRadius = FMath::Max(*MaxRadius, TypeMaxCollisionRadius);
@@ -238,7 +238,7 @@ void UProceduralFoliageTile::AddRandomSeeds(TArray<FProceduralFoliageInstance*>&
 					MaxCollisionRadii.Add(DistributionSeed, TypeMaxCollisionRadius);
 				}
 
-				const float TypeMaxShadeRadius = MaxScale * TypeInstance->ShadeRadius;
+				const float TypeMaxShadeRadius = MaxScale * Type->ShadeRadius;
 				if (float* MaxRadius = MaxShadeRadii.Find(DistributionSeed))
 				{
 					*MaxRadius = FMath::Max(*MaxRadius, TypeMaxShadeRadius);
@@ -411,9 +411,9 @@ void UProceduralFoliageTile::RunSimulation(const int32 MaxNumSteps, bool bShadeG
 {
 	int32 MaxSteps = 0;
 
-	for (const FFoliageTypeObject& FoliageTypeObject : FoliageSpawner->GetFoliageTypes())
+	for (const FProceduralFoliageTypeData& Data : FoliageSpawner->GetTypes())
 	{
-		const UFoliageType_InstancedStaticMesh* TypeInstance = FoliageTypeObject.GetInstance();
+		const UFoliageType_InstancedStaticMesh* TypeInstance = Data.TypeInstance;
 		if (TypeInstance && TypeInstance->bGrowsInShade == bShadeGrowth)
 		{
 			MaxSteps = FMath::Max(MaxSteps, TypeInstance->NumSteps + 1);
diff --git a/Engine/Source/Runtime/Foliage/Public/FoliageTypeObject.h b/Engine/Source/Runtime/Foliage/Public/FoliageTypeObject.h
deleted file mode 100644
diff --git a/Engine/Source/Runtime/Foliage/Public/FoliageType_InstancedStaticMesh.h b/Engine/Source/Runtime/Foliage/Public/FoliageType_InstancedStaticMesh.h
--- a/Engine/Source/Runtime/Foliage/Public/FoliageType_InstancedStaticMesh.h
+++ b/Engine/Source/Runtime/Foliage/Public/FoliageType_InstancedStaticMesh.h
@@ -10,7 +10,7 @@ class UFoliageType_InstancedStaticMesh : public UFoliageType
 {
 	GENERATED_UCLASS_BODY()
 
-	UPROPERTY(EditAnywhere, Category=Mesh)
+	UPROPERTY(EditAnywhere, Category=FoliageType)
 	UStaticMesh* Mesh;
 
 	virtual UStaticMesh* GetStaticMesh() const override
diff --git a/Engine/Source/Runtime/Foliage/Public/ProceduralFoliageSpawner.h b/Engine/Source/Runtime/Foliage/Public/ProceduralFoliageSpawner.h
--- a/Engine/Source/Runtime/Foliage/Public/ProceduralFoliageSpawner.h
+++ b/Engine/Source/Runtime/Foliage/Public/ProceduralFoliageSpawner.h
@@ -4,11 +4,26 @@
 
 #include "FoliageType_InstancedStaticMesh.h"
 #include "ProceduralFoliageInstance.h"
-#include "FoliageTypeObject.h"
 #include "ProceduralFoliageSpawner.generated.h"
 
 class UProceduralFoliageTile;
 
+
+USTRUCT()
+struct FProceduralFoliageTypeData
+{
+	GENERATED_USTRUCT_BODY()
+
+	UPROPERTY(Category = ProceduralFoliageSimulation, EditAnywhere, BlueprintReadOnly)
+	TSubclassOf<UFoliageType_InstancedStaticMesh> Type;
+
+	UPROPERTY(transient)
+	UFoliageType_InstancedStaticMesh* TypeInstance;
+
+	UPROPERTY()
+	int32 ChangeCount;
+};
+
 UCLASS(BlueprintType, Blueprintable)
 class FOLIAGE_API UProceduralFoliageSpawner : public UObject
 {
@@ -31,7 +46,7 @@ class FOLIAGE_API UProceduralFoliageSpawner : public UObject
 private:
 	/** The types of foliage to procedurally spawn. */
 	UPROPERTY(Category = ProceduralFoliageSimulation, EditAnywhere)
-	TArray<FFoliageTypeObject> FoliageTypes;
+	TArray<FProceduralFoliageTypeData> Types;
 
 	UPROPERTY()
 	bool bNeedsSimulation;
@@ -42,7 +57,7 @@ public:
 
 	int32 GetRandomNumber();
 
-	const TArray<FFoliageTypeObject>& GetFoliageTypes() const { return FoliageTypes; }
+	const TArray<FProceduralFoliageTypeData>& GetTypes() const { return Types; }
 
 	/** Returns the instances that need to spawn for a given min,max region */
 	void GetInstancesToSpawn(TArray<FProceduralFoliageInstance>& OutInstances, const FVector& Min = FVector(-FLT_MAX, -FLT_MAX, -FLT_MAX), const FVector& Max = FVector(FLT_MAX, FLT_MAX, FLT_MAX) ) const;
