Fixed:  updating cached transforms for foliage instance base components

diff --git a/Engine/Source/Runtime/Engine/Classes/Engine/World.h b/Engine/Source/Runtime/Engine/Classes/Engine/World.h
--- a/Engine/Source/Runtime/Engine/Classes/Engine/World.h
+++ b/Engine/Source/Runtime/Engine/Classes/Engine/World.h
@@ -3226,6 +3226,10 @@ public:
 	// removed from the world without a RemoveFromWorld call for each)
 	static FOnLevelChanged			LevelRemovedFromWorld;
 
+	// Called after offset was applied to a level
+	DECLARE_MULTICAST_DELEGATE_FourParams(FLevelOffsetEvent, ULevel*,  UWorld*, const FVector&, bool);
+	static FLevelOffsetEvent		PostApplyLevelOffset;
+
 	// called by UWorld::GetAssetRegistryTags()
 	static FWorldGetAssetTags GetAssetTags;
 
diff --git a/Engine/Source/Runtime/Engine/Private/Level.cpp b/Engine/Source/Runtime/Engine/Private/Level.cpp
--- a/Engine/Source/Runtime/Engine/Private/Level.cpp
+++ b/Engine/Source/Runtime/Engine/Private/Level.cpp
@@ -1851,6 +1851,8 @@ void ULevel::ApplyWorldOffset(const FVector& InWorldOffset, bool bWorldShift)
 			ModelComponents[CompIdx]->ApplyWorldOffset(InWorldOffset, bWorldShift);
 		}
 	}
+
+	FWorldDelegates::PostApplyLevelOffset.Broadcast(this, OwningWorld, InWorldOffset, bWorldShift);
 }
 
 void ULevel::RegisterActorForAutoReceiveInput(AActor* Actor, const int32 PlayerIndex)
diff --git a/Engine/Source/Runtime/Engine/Private/World.cpp b/Engine/Source/Runtime/Engine/Private/World.cpp
--- a/Engine/Source/Runtime/Engine/Private/World.cpp
+++ b/Engine/Source/Runtime/Engine/Private/World.cpp
@@ -103,6 +103,7 @@ FWorldDelegates::FWorldCleanupEvent FWorldDelegates::OnWorldCleanup;
 FWorldDelegates::FWorldEvent FWorldDelegates::OnPreWorldFinishDestroy;
 FWorldDelegates::FOnLevelChanged FWorldDelegates::LevelAddedToWorld;
 FWorldDelegates::FOnLevelChanged FWorldDelegates::LevelRemovedFromWorld;
+FWorldDelegates::FLevelOffsetEvent FWorldDelegates::PostApplyLevelOffset;
 FWorldDelegates::FWorldGetAssetTags FWorldDelegates::GetAssetTags;
 #if WITH_EDITOR
 FWorldDelegates::FLevelScriptCreationEvent FWorldDelegates::LevelScriptBlueprintCreated;
diff --git a/Engine/Source/Runtime/Foliage/Private/InstancedFoliage.cpp b/Engine/Source/Runtime/Foliage/Private/InstancedFoliage.cpp
--- a/Engine/Source/Runtime/Foliage/Private/InstancedFoliage.cpp
+++ b/Engine/Source/Runtime/Foliage/Private/InstancedFoliage.cpp
@@ -2090,6 +2090,9 @@ void AInstancedFoliageActor::PostInitProperties()
 
 		GEngine->OnLevelActorDeleted().Remove(OnLevelActorDeletedDelegateHandle);
 		OnLevelActorDeletedDelegateHandle = GEngine->OnLevelActorDeleted().AddUObject(this, &AInstancedFoliageActor::OnLevelActorDeleted);
+
+		FWorldDelegates::PostApplyLevelOffset.Remove(OnPostApplyLevelOffsetDelegateHandle);
+		OnPostApplyLevelOffsetDelegateHandle = FWorldDelegates::PostApplyLevelOffset.AddUObject(this, &AInstancedFoliageActor::OnPostApplyLevelOffset);
 	}
 #endif
 }
@@ -2103,6 +2106,7 @@ void AInstancedFoliageActor::BeginDestroy()
 	{
 		GEngine->OnActorMoved().Remove(OnLevelActorMovedDelegateHandle);
 		GEngine->OnLevelActorDeleted().Remove(OnLevelActorDeletedDelegateHandle);
+		FWorldDelegates::PostApplyLevelOffset.Remove(OnPostApplyLevelOffsetDelegateHandle);
 	}
 #endif
 }
@@ -2259,37 +2263,15 @@ void AInstancedFoliageActor::OnLevelActorDeleted(AActor* InActor)
 	}
 }
 
-#endif
-
-//
-// Serialize all our UObjects for RTGC 
-//
-void AInstancedFoliageActor::AddReferencedObjects(UObject* InThis, FReferenceCollector& Collector)
+void AInstancedFoliageActor::OnPostApplyLevelOffset(ULevel* InLevel, UWorld* InWorld, const FVector& InOffset, bool bWorldShift)
 {
-	AInstancedFoliageActor* This = CastChecked<AInstancedFoliageActor>(InThis);
-
-	for (auto& MeshPair : This->FoliageMeshes)
+	ULevel* OwningLevel = GetLevel();
+	if (InLevel != OwningLevel) // TODO: cross-level foliage 
 	{
-		Collector.AddReferencedObject(MeshPair.Key, This);
-		FFoliageMeshInfo& MeshInfo = *MeshPair.Value;
-
-		if (MeshInfo.Component)
-		{
-			Collector.AddReferencedObject(MeshInfo.Component, This);
-		}
+		return;
 	}
 	
-	Super::AddReferencedObjects(This, Collector);
-}
-
-void AInstancedFoliageActor::ApplyWorldOffset(const FVector& InOffset, bool bWorldShift)
-{
-	Super::ApplyWorldOffset(InOffset, bWorldShift);
-
-#if WITH_EDITORONLY_DATA	
-	UWorld* World = GetWorld();
-
-	if (GIsEditor && World && !World->IsGameWorld())
+	if (GIsEditor && InWorld && !InWorld->IsGameWorld())
 	{
 		for (auto& MeshPair : FoliageMeshes)
 		{
@@ -2307,7 +2289,29 @@ void AInstancedFoliageActor::ApplyWorldOffset(const FVector& InOffset, bool bWor
 			}
 		}
 	}
+}
+
 #endif
+
+//
+// Serialize all our UObjects for RTGC 
+//
+void AInstancedFoliageActor::AddReferencedObjects(UObject* InThis, FReferenceCollector& Collector)
+{
+	AInstancedFoliageActor* This = CastChecked<AInstancedFoliageActor>(InThis);
+
+	for (auto& MeshPair : This->FoliageMeshes)
+	{
+		Collector.AddReferencedObject(MeshPair.Key, This);
+		FFoliageMeshInfo& MeshInfo = *MeshPair.Value;
+
+		if (MeshInfo.Component)
+		{
+			Collector.AddReferencedObject(MeshInfo.Component, This);
+		}
+	}
+	
+	Super::AddReferencedObjects(This, Collector);
 }
 
 bool AInstancedFoliageActor::FoliageTrace(const UWorld* InWorld, FHitResult& OutHit, const FDesiredFoliageInstance& DesiredInstance, FName InTraceTag, bool InbReturnFaceIndex)
diff --git a/Engine/Source/Runtime/Foliage/Public/InstancedFoliageActor.h b/Engine/Source/Runtime/Foliage/Public/InstancedFoliageActor.h
--- a/Engine/Source/Runtime/Foliage/Public/InstancedFoliageActor.h
+++ b/Engine/Source/Runtime/Foliage/Public/InstancedFoliageActor.h
@@ -45,7 +45,6 @@ public:
 	// Begin AActor interface.
 	// we don't want to have our components automatically destroyed by the Blueprint code
 	virtual void RerunConstructionScripts() override {}
-	virtual void ApplyWorldOffset(const FVector& InOffset, bool bWorldShift) override;
 	virtual bool IsLevelBoundsRelevant() const override { return false; }
 #if WITH_EDITOR
 	virtual void Destroyed() override;
@@ -171,11 +170,13 @@ private:
 #if WITH_EDITOR
 	void OnLevelActorMoved(AActor* InActor);
 	void OnLevelActorDeleted(AActor* InActor);
+	void OnPostApplyLevelOffset(ULevel* InLevel, UWorld* InWorld, const FVector& InOffset, bool bWorldShift);
 #endif
 private:
 #if WITH_EDITOR
 	FDelegateHandle OnLevelActorMovedDelegateHandle;
 	FDelegateHandle OnLevelActorDeletedDelegateHandle;
+	FDelegateHandle OnPostApplyLevelOffsetDelegateHandle;
 #endif
 
 };
