Do not copy foliage with empty base when copy/pasting actors

diff --git a/Engine/Source/Runtime/Foliage/Private/InstancedFoliage.cpp b/Engine/Source/Runtime/Foliage/Private/InstancedFoliage.cpp
--- a/Engine/Source/Runtime/Foliage/Private/InstancedFoliage.cpp
+++ b/Engine/Source/Runtime/Foliage/Private/InstancedFoliage.cpp
@@ -1529,20 +1529,23 @@ TMap<UFoliageType*, TArray<const FFoliageInstancePlacementInfo*>> AInstancedFoli
 {
 	TMap<UFoliageType*, TArray<const FFoliageInstancePlacementInfo*>> Result;
 	const auto BaseId = InstanceBaseCache.GetInstanceBaseId(InComponent);
-
-	for (auto& MeshPair : FoliageMeshes)
+	
+	if (BaseId != FFoliageInstanceBaseCache::InvalidBaseId)
 	{
-		const FFoliageMeshInfo& MeshInfo = *MeshPair.Value;
-		const auto* InstanceSet = MeshInfo.ComponentHash.Find(BaseId);
-		if (InstanceSet)
+		for (auto& MeshPair : FoliageMeshes)
 		{
-			TArray<const FFoliageInstancePlacementInfo*>& Array = Result.Add(MeshPair.Key, TArray<const FFoliageInstancePlacementInfo*>());
-			Array.Empty(InstanceSet->Num());
-
-			for (int32 InstanceIndex : *InstanceSet)
+			const FFoliageMeshInfo& MeshInfo = *MeshPair.Value;
+			const auto* InstanceSet = MeshInfo.ComponentHash.Find(BaseId);
+			if (InstanceSet)
 			{
-				const FFoliageInstancePlacementInfo* Instance = &MeshInfo.Instances[InstanceIndex];
-				Array.Add(Instance);
+				TArray<const FFoliageInstancePlacementInfo*>& Array = Result.Add(MeshPair.Key, TArray<const FFoliageInstancePlacementInfo*>());
+				Array.Empty(InstanceSet->Num());
+
+				for (int32 InstanceIndex : *InstanceSet)
+				{
+					const FFoliageInstancePlacementInfo* Instance = &MeshInfo.Instances[InstanceIndex];
+					Array.Add(Instance);
+				}
 			}
 		}
 	}
