Fixed crash when selecting a foliage instance if a type in the palette had no instances in the world

diff --git a/Engine/Source/Runtime/Foliage/Private/InstancedFoliage.cpp b/Engine/Source/Runtime/Foliage/Private/InstancedFoliage.cpp
--- a/Engine/Source/Runtime/Foliage/Private/InstancedFoliage.cpp
+++ b/Engine/Source/Runtime/Foliage/Private/InstancedFoliage.cpp
@@ -1746,10 +1746,13 @@ void AInstancedFoliageActor::SelectInstance(UInstancedStaticMeshComponent* InCom
 		{
 			FFoliageMeshInfo& MeshInfo = *MeshPair.Value;
 
-			check(MeshInfo.Component);
-			MeshInfo.Component->ClearInstanceSelection();
-			MeshInfo.Component->MarkRenderStateDirty();
-			MeshInfo.SelectedIndices.Empty();
+			if (MeshInfo.Instances.Num() > 0)
+			{
+				check(MeshInfo.Component);
+				MeshInfo.Component->ClearInstanceSelection();
+				MeshInfo.Component->MarkRenderStateDirty();
+				MeshInfo.SelectedIndices.Empty();
+			}	
 		}
 	}
 
