Added edit conditions to various shadow flags on the instance settings

diff --git a/Engine/Source/Runtime/Foliage/Public/FoliageType.h b/Engine/Source/Runtime/Foliage/Public/FoliageType.h
--- a/Engine/Source/Runtime/Foliage/Public/FoliageType.h
+++ b/Engine/Source/Runtime/Foliage/Public/FoliageType.h
@@ -187,7 +187,7 @@ public:
 	uint32 CastShadow:1;
 
 	/** Controls whether the foliage should inject light into the Light Propagation Volume.  This flag is only used if CastShadow is true. */
-	UPROPERTY(EditAnywhere, BlueprintReadOnly, Category=InstanceSettings)
+	UPROPERTY(EditAnywhere, BlueprintReadOnly, Category=InstanceSettings, meta=(EditCondition="CastShadow"))
 	uint32 bAffectDynamicIndirectLighting:1;
 
 	/** Controls whether the primitive should affect dynamic distance field lighting methods.  This flag is only used if CastShadow is true. */
@@ -195,15 +195,15 @@ public:
 	uint32 bAffectDistanceFieldLighting:1;
 
 	/** Controls whether the foliage should cast shadows in the case of non precomputed shadowing.  This flag is only used if CastShadow is true. */
-	UPROPERTY(EditAnywhere, BlueprintReadOnly, Category=InstanceSettings)
+	UPROPERTY(EditAnywhere, BlueprintReadOnly, Category=InstanceSettings, meta=(EditCondition="CastShadow"))
 	uint32 bCastDynamicShadow:1;
 
 	/** Whether the foliage should cast a static shadow from shadow casting lights.  This flag is only used if CastShadow is true. */
-	UPROPERTY(EditAnywhere, BlueprintReadOnly, Category=InstanceSettings)
+	UPROPERTY(EditAnywhere, BlueprintReadOnly, Category=InstanceSettings, meta=(EditCondition="CastShadow"))
 	uint32 bCastStaticShadow:1;
 
 	/** Whether this foliage should cast dynamic shadows as if it were a two sided material. */
-	UPROPERTY(EditAnywhere, AdvancedDisplay, BlueprintReadOnly, Category=InstanceSettings)
+	UPROPERTY(EditAnywhere, AdvancedDisplay, BlueprintReadOnly, Category=InstanceSettings, meta=(EditCondition="bCastDynamicShadow"))
 	uint32 bCastShadowAsTwoSided:1;
 
 	/** Whether the foliage receives decals. */
