[UE-14606] - Updated shade growth behavior - now possible to be able to grow in shade without being required to spawn in shade. Old bGrowsInShade option is now bSpawnsInShade.
[UE-14605] - Painting scale only visible when painting and proc scale only visible when the asset is open
[UE-14155] - Fixed procedural generation to not add redundant entries for foliage types when simulating. When Foliage Type assets are not used, the foliage type created by the proc foliage simulation will stomp an existing non-asset type with the same mesh (if there is one) in the palette. A significant rework of the foliage palette is needed for a more robust solution.

- Fixed crash when an instance is located exactly on the border between two quadtrees

diff --git a/Engine/Config/BaseEngine.ini b/Engine/Config/BaseEngine.ini
--- a/Engine/Config/BaseEngine.ini
+++ b/Engine/Config/BaseEngine.ini
@@ -362,6 +362,8 @@ bUseStreamingPause=false
 +TaggedPropertyRedirects=(ClassName="ProceduralFoliageComponent",OldPropertyName="Overlap",NewPropertyName="TileOverlap")
 +ActiveStructRedirects=(OldStructName="ProceduralFoliageTypeData",NewStructName="FoliageTypeObject")
 +TaggedPropertyRedirects=(ClassName="ProceduralFoliageSpawner", OldPropertyName="Types", NewPropertyName="FoliageTypes")
++TaggedPropertyRedirects=(ClassName="FoliageType", OldPropertyName="InitialMaxAge", NewPropertyName="MaxInitialAge")
++TaggedPropertyRedirects=(ClassName="FoliageType", OldPropertyName="bGrowsInShade", NewPropertyName="bSpawnsInShade")
 
 ;Shortened animation compression algorithm class names
 +ActiveClassRedirects=(OldClassName="AnimationCompressionAlgorithm",NewClassName="AnimCompress")
@@ -1209,9 +1211,6 @@ bUseStreamingPause=false
 +ActiveClassRedirects=(OldClassName="InstancedFoliageActor", NewClassName="/Script/Foliage.InstancedFoliageActor")
 +ActiveClassRedirects=(OldClassName="InteractiveFoliageComponent", NewClassName="/Script/Foliage.InteractiveFoliageComponent")
 +ActiveClassRedirects=(OldClassName="FoliageVertexColorMask", NewClassName="/Script/Foliage.FoliageVertexColorMask")
-+TaggedPropertyRedirects(ClassName="FoliageType", OldPropertyName="LandscapeLayer", NewPropertyName="LandscapeLayer_DEPRECATED")
-
-
 
 ;Removal of EmitterSpawnable
 +ActiveClassRedirects=(OldClassName="EmitterSpawnable", NewClassName="Emitter")
diff --git a/Engine/Source/Editor/FoliageEdit/Private/FoliageEdMode.cpp b/Engine/Source/Editor/FoliageEdit/Private/FoliageEdMode.cpp
--- a/Engine/Source/Editor/FoliageEdit/Private/FoliageEdMode.cpp
+++ b/Engine/Source/Editor/FoliageEdit/Private/FoliageEdMode.cpp
@@ -40,6 +40,8 @@
 #define LOCTEXT_NAMESPACE "FoliageEdMode"
 #define FOLIAGE_SNAP_TRACE (10000.f)
 
+DEFINE_LOG_CATEGORY_STATIC(LogFoliage, Log, Warning);
+
 //
 // FFoliageMeshUIInfo
 //
@@ -940,17 +942,28 @@ FFoliageMeshInfo* UpdateMeshSettings(AInstancedFoliageActor* IFA, const FDesired
 {
 	FFoliageMeshInfo* MeshInfo = nullptr;
 
-	//We have to add the foliage type into the foliage actor.
 	if (const UFoliageType_InstancedStaticMesh* FoliageType = Cast<UFoliageType_InstancedStaticMesh>(DesiredInstance.FoliageType))
 	{
-		if (UStaticMesh* StaticMesh = FoliageType->GetStaticMesh())
+		//We have to add the foliage type into the foliage actor.
+		if (FoliageType->IsAsset())
+		{
+			// Add the foliage type asset
+			IFA->AddFoliageType(FoliageType, &MeshInfo);
+		}
+		else if (UStaticMesh* StaticMesh = FoliageType->GetStaticMesh())
 		{
-			if (UFoliageType* ExistingType = IFA->GetSettingsForMesh(StaticMesh))
+			if (IFA->GetSettingsForMesh(StaticMesh, nullptr, false) != nullptr)
 			{
-				MeshInfo = IFA->UpdateMeshSettings(StaticMesh, FoliageType);
+				// @todo: create some unique designation for types added procedurally to prevent them stomping any existing types (or restrict proc foliage to foliage type assets only)
+
+				// A non-asset foliage type using this mesh already exists in the IFA
+				// Only one such type is supported, so the desired instance overrides it
+				MeshInfo = IFA->UpdateMeshSettings(StaticMesh, FoliageType, false);
+				UE_LOG(LogFoliage, Warning, TEXT("Procedurally generated foliage type using Static Mesh \"%s\" conflicts with an existing non-asset foliage type. To prevent clashes, use Foliage Type assets instead."), *StaticMesh->GetName());
 			}
 			else
 			{
+				// No conflict with an existing type, just add normally
 				IFA->AddFoliageType(FoliageType, &MeshInfo);
 			}
 		}
@@ -976,22 +989,21 @@ void FEdModeFoliage::AddInstancesImp(UWorld* InWorld, const UFoliageType* Settin
 		//@TODO: actual threaded part coming, need parts of this refactor sooner for content team
 		CalculatePotentialInstances_ThreadSafe(InWorld, Settings, &DesiredInstances, PotentialInstanceBuckets, nullptr, 0, DesiredInstances.Num() - 1);
 
-		TMap<AInstancedFoliageActor*, FFoliageMeshInfo*> UpdatedIFAs;	//we want to override any existing mesh settings with the procedural settings.
-		for (auto& Bucket : PotentialInstanceBuckets)
+		// Existing foliage types in the palette  we want to override any existing mesh settings with the procedural settings.
+		TMap<AInstancedFoliageActor*, TArray<const UFoliageType*>> UpdatedTypesByIFA;
+		for (TArray<FPotentialInstance>& Bucket : PotentialInstanceBuckets)
 		{
-			for (FPotentialInstance& PotentialInst : Bucket)
+			for (auto& PotentialInst : Bucket)
 			{
-				// New foliage instances always go into base component level
+				// Get the IFA for the base component level that contains the component the instance will be placed upon
 				AInstancedFoliageActor* TargetIFA = AInstancedFoliageActor::GetInstancedFoliageActorForLevel(PotentialInst.HitComponent->GetComponentLevel(), true);
-				FFoliageMeshInfo* MeshInfo = nullptr;
-				if (FFoliageMeshInfo** MeshInfoResult = UpdatedIFAs.Find(TargetIFA))
-				{
-					MeshInfo = *MeshInfoResult;
-				}
-				else
+				
+				// Update the type in the IFA if needed
+				TArray<const UFoliageType*>& UpdatedTypes = UpdatedTypesByIFA.FindOrAdd(TargetIFA);
+				if (!UpdatedTypes.Contains(PotentialInst.DesiredInstance.FoliageType))
 				{
-					MeshInfo = UpdateMeshSettings(TargetIFA, PotentialInst.DesiredInstance);
-					UpdatedIFAs.Add(TargetIFA, MeshInfo);
+					UpdatedTypes.Add(PotentialInst.DesiredInstance.FoliageType);
+					UpdateMeshSettings(TargetIFA, PotentialInst.DesiredInstance);
 				}
 			}
 		}
diff --git a/Engine/Source/Editor/FoliageEdit/Private/FoliageTypeDetails.cpp b/Engine/Source/Editor/FoliageEdit/Private/FoliageTypeDetails.cpp
--- a/Engine/Source/Editor/FoliageEdit/Private/FoliageTypeDetails.cpp
+++ b/Engine/Source/Editor/FoliageEdit/Private/FoliageTypeDetails.cpp
@@ -32,7 +32,7 @@ void CustomizePropertyRowVisibility(IDetailLayoutBuilder& LayoutBuilder, const T
 	}
 }
 
-void AddSubcategoryProperties(IDetailLayoutBuilder& LayoutBuilder, const FName CategoryName, TMap<const FName, IDetailPropertyRow*>* OutDetailRowsByPropertyName = nullptr)
+void AddSubcategoryProperties(IDetailLayoutBuilder& LayoutBuilder, const FName CategoryName)
 {
 	IDetailCategoryBuilder& CategoryBuilder = LayoutBuilder.EditCategory(CategoryName);
 
@@ -57,11 +57,6 @@ void AddSubcategoryProperties(IDetailLayoutBuilder& LayoutBuilder, const FName C
 				// The property is not in a subcategory, so add it now
 				auto& PropertyRow = CategoryBuilder.AddProperty(PropertyHandle);
 				CustomizePropertyRowVisibility(LayoutBuilder, PropertyHandle, PropertyRow);
-
-				if (OutDetailRowsByPropertyName)
-				{
-					OutDetailRowsByPropertyName->Add(Property->GetFName()) = &PropertyRow;
-				}
 			}
 		}
 	}
@@ -76,11 +71,6 @@ void AddSubcategoryProperties(IDetailLayoutBuilder& LayoutBuilder, const FName C
 		{
 			auto& PropertyRow = Group.AddPropertyRow(PropertyHandle);
 			CustomizePropertyRowVisibility(LayoutBuilder, PropertyHandle, PropertyRow);
-
-			if (OutDetailRowsByPropertyName && PropertyHandle->GetProperty())
-			{
-				OutDetailRowsByPropertyName->Add(PropertyHandle->GetProperty()->GetFName()) = &PropertyRow;
-			}
 		}
 	}
 }
@@ -96,8 +86,7 @@ void FFoliageTypeDetails::CustomizeDetails(IDetailLayoutBuilder& DetailBuilder)
 	FFoliageTypeCustomizationHelpers::HideFoliageCategory(DetailBuilder, ReapplyName);
 	FFoliageTypeCustomizationHelpers::HideFoliageCategory(DetailBuilder, PaintingName);
 
-	TMap<const FName, IDetailPropertyRow*> PropertyRowsByName;
-	AddSubcategoryProperties(DetailBuilder, PlacementName, &PropertyRowsByName);
+	AddSubcategoryProperties(DetailBuilder, PlacementName);
 
 	// If enabled, show the properties for procedural placement
 	if (GetDefault<UEditorExperimentalSettings>()->bProceduralFoliage)
@@ -110,19 +99,4 @@ void FFoliageTypeDetails::CustomizeDetails(IDetailLayoutBuilder& DetailBuilder)
 	}
 	
 	AddSubcategoryProperties(DetailBuilder, InstanceSettingsName);
-
-	// Custom visibility overrides for the scale axes
-	TSharedPtr<IPropertyHandle> Scaling = DetailBuilder.GetProperty(GET_MEMBER_NAME_CHECKED(UFoliageType, Scaling));
-
-	FFoliageTypeCustomizationHelpers::ModifyFoliagePropertyRow(*PropertyRowsByName.Find(GET_MEMBER_NAME_CHECKED(UFoliageType, ScaleX)),
-		TAttribute<EVisibility>::Create(TAttribute<EVisibility>::FGetter::CreateStatic(&FFoliageTypeCustomizationHelpers::GetScaleAxisVisibility, EAxis::X, Scaling)),
-		TAttribute<bool>());
-
-	FFoliageTypeCustomizationHelpers::ModifyFoliagePropertyRow(*PropertyRowsByName.Find(GET_MEMBER_NAME_CHECKED(UFoliageType, ScaleY)),
-		TAttribute<EVisibility>::Create(TAttribute<EVisibility>::FGetter::CreateStatic(&FFoliageTypeCustomizationHelpers::GetScaleAxisVisibility, EAxis::Y, Scaling)),
-		TAttribute<bool>());
-
-	FFoliageTypeCustomizationHelpers::ModifyFoliagePropertyRow(*PropertyRowsByName.Find(GET_MEMBER_NAME_CHECKED(UFoliageType, ScaleZ)),
-		TAttribute<EVisibility>::Create(TAttribute<EVisibility>::FGetter::CreateStatic(&FFoliageTypeCustomizationHelpers::GetScaleAxisVisibility, EAxis::Z, Scaling)),
-		TAttribute<bool>());
 }
\ No newline at end of file
diff --git a/Engine/Source/Runtime/Engine/Public/GenericQuadTree.h b/Engine/Source/Runtime/Engine/Public/GenericQuadTree.h
--- a/Engine/Source/Runtime/Engine/Public/GenericQuadTree.h
+++ b/Engine/Source/Runtime/Engine/Public/GenericQuadTree.h
@@ -359,11 +359,11 @@ int32 TQuadTree<ElementType, NodeCapacity>::GetQuads(const FBox2D& Box, TreeType
 	int32 QuadCount = 0;
 	if (bInternal)
 	{
-		bool bNegX = Box.Min.X < Position.X;
-		bool bNegY = Box.Min.Y < Position.Y;
+		bool bNegX = Box.Min.X <= Position.X;
+		bool bNegY = Box.Min.Y <= Position.Y;
 
-		bool bPosX = Box.Max.X > Position.X;
-		bool bPosY = Box.Max.Y > Position.Y;
+		bool bPosX = Box.Max.X >= Position.X;
+		bool bPosY = Box.Max.Y >= Position.Y;
 
 		if (bNegX && bNegY)
 		{
diff --git a/Engine/Source/Runtime/Foliage/Private/InstancedFoliage.cpp b/Engine/Source/Runtime/Foliage/Private/InstancedFoliage.cpp
--- a/Engine/Source/Runtime/Foliage/Private/InstancedFoliage.cpp
+++ b/Engine/Source/Runtime/Foliage/Private/InstancedFoliage.cpp
@@ -46,6 +46,8 @@ struct FFoliageCustomVersion
 		FoliageTypeCustomization = 6,
 		// FoliageType for details customization continued
 		FoliageTypeCustomizationScaling = 7,
+		// FoliageType procedural scale and shade settings updated
+		FoliageTypeProceduralScaleAndShade = 8,
 		// -----<new versions can be added above this line>-------------------------------------------------
 		VersionPlusOne,
 		LatestVersion = VersionPlusOne - 1
@@ -305,17 +307,17 @@ UFoliageType::UFoliageType(const FObjectInitializer& ObjectInitializer)
 	/** Ecosystem settings*/
 	AverageSpreadDistance = 50;
 	SpreadVariance = 150;
-	bGrowsInShade = false;
+	bCanGrowInShade = false;
+	bSpawnsInShade = false;
 	SeedsPerStep = 3;
 	OverlapPriority = 0.f;
 	NumSteps = 3;
-	MinScale = 1.f;
-	MaxScale = 3.f;
+	ProceduralScale = FFloatInterval(1.f, 3.f);
 	ChangeCount = 0;
 	InitialSeedDensity = 1.f;
 	CollisionRadius = 100.f;
 	ShadeRadius = 100.f;
-	InitialMaxAge = 0.f;
+	MaxInitialAge = 0.f;
 	MaxAge = 10.f;
 
 	FRichCurve* Curve = ScaleCurve.GetRichCurve();
@@ -324,8 +326,8 @@ UFoliageType::UFoliageType(const FObjectInitializer& ObjectInitializer)
 
 	UpdateGuid = FGuid::NewGuid();
 
-	// Deprecated since FFoliageCustomVersion::FoliageTypeCustomization
 #if WITH_EDITORONLY_DATA
+	// Deprecated since FFoliageCustomVersion::FoliageTypeCustomization
 	ScaleMinX_DEPRECATED = 1.0f;
 	ScaleMinY_DEPRECATED = 1.0f;
 	ScaleMinZ_DEPRECATED = 1.0f;
@@ -338,6 +340,11 @@ UFoliageType::UFoliageType(const FObjectInitializer& ObjectInitializer)
 	ZOffsetMax_DEPRECATED = 0.0f;
 	UniformScale_DEPRECATED = true;
 	GroundSlope_DEPRECATED = 45.0f;
+
+	// Deprecated since FFoliageCustomVersion::FoliageTypeProceduralScaleAndShade
+	MinScale_DEPRECATED = 1.f;
+	MaxScale_DEPRECATED = 3.f;
+
 #endif// WITH_EDITORONLY_DATA
 }
 
@@ -384,6 +391,14 @@ void UFoliageType::Serialize(FArchive& Ar)
 			GroundSlopeAngle.Min = MinGroundSlope_DEPRECATED;
 			GroundSlopeAngle.Max = GroundSlope_DEPRECATED;
 		}
+
+		if (Ar.CustomVer(FFoliageCustomVersion::GUID) < FFoliageCustomVersion::FoliageTypeProceduralScaleAndShade)
+		{
+			bCanGrowInShade = bSpawnsInShade;
+
+			ProceduralScale.Min = MinScale_DEPRECATED;
+			ProceduralScale.Max = MaxScale_DEPRECATED;
+		}
 	}
 #endif// WITH_EDITORONLY_DATA
 }
@@ -447,12 +462,12 @@ float UFoliageType::GetScaleForAge(const float Age) const
 	const FRichCurve* Curve = ScaleCurve.GetRichCurveConst();
 	const float Time = FMath::Clamp(MaxAge == 0 ? 1.f : Age / MaxAge, 0.f, 1.f);
 	const float Scale = Curve->Eval(Time);
-	return MinScale + (MaxScale - MinScale) * Scale;
+	return ProceduralScale.Min + ProceduralScale.Size() * Scale;
 }
 
 float UFoliageType::GetInitAge(FRandomStream& RandomStream) const
 {
-	return RandomStream.FRandRange(0, InitialMaxAge);
+	return RandomStream.FRandRange(0, MaxInitialAge);
 }
 
 float UFoliageType::GetNextAge(const float CurrentAge, const int32 InNumSteps) const
@@ -474,6 +489,11 @@ float UFoliageType::GetNextAge(const float CurrentAge, const int32 InNumSteps) c
 	return NewAge;
 }
 
+bool UFoliageType::GetSpawnsInShade() const
+{
+	return bCanGrowInShade && bSpawnsInShade;
+}
+
 #if WITH_EDITOR
 void UFoliageType::PostEditChangeProperty(struct FPropertyChangedEvent& PropertyChangedEvent)
 {
@@ -1200,17 +1220,17 @@ void AInstancedFoliageActor::GetOverlappingMeshCounts(const FSphere& Sphere, TMa
 
 
 
-UFoliageType* AInstancedFoliageActor::GetSettingsForMesh(const UStaticMesh* InMesh, FFoliageMeshInfo** OutMeshInfo)
+UFoliageType* AInstancedFoliageActor::GetSettingsForMesh(const UStaticMesh* InMesh, FFoliageMeshInfo** OutMeshInfo, bool bIncludeAssets)
 {
-	UFoliageType* Type = nullptr;
+	UFoliageType* ReturnType = nullptr;
 	FFoliageMeshInfo* MeshInfo = nullptr;
 
 	for (auto& MeshPair : FoliageMeshes)
 	{
-		UFoliageType* Settings = MeshPair.Key;
-		if (Settings && Settings->GetStaticMesh() == InMesh)
+		UFoliageType* FoliageType = MeshPair.Key;
+		if (FoliageType && (bIncludeAssets || !FoliageType->IsAsset()) && FoliageType->GetStaticMesh() == InMesh)
 		{
-			Type = MeshPair.Key;
+			ReturnType = FoliageType;
 			MeshInfo = &*MeshPair.Value;
 			break;
 		}
@@ -1220,7 +1240,8 @@ UFoliageType* AInstancedFoliageActor::GetSettingsForMesh(const UStaticMesh* InMe
 	{
 		*OutMeshInfo = MeshInfo;
 	}
-	return Type;
+
+	return ReturnType;
 }
 
 
@@ -1606,9 +1627,9 @@ void UpdateSettingsBounds(const UStaticMesh* InMesh, UFoliageType_InstancedStati
 
 
 #if WITH_EDITORONLY_DATA
-FFoliageMeshInfo* AInstancedFoliageActor::UpdateMeshSettings(const UStaticMesh* InMesh, const UFoliageType_InstancedStaticMesh* DefaultSettings)
+FFoliageMeshInfo* AInstancedFoliageActor::UpdateMeshSettings(const UStaticMesh* InMesh, const UFoliageType_InstancedStaticMesh* DefaultSettings, bool bIncludeAssets)
 {
-	if (UFoliageType* OldSettings = GetSettingsForMesh(InMesh))
+	if (UFoliageType* OldSettings = GetSettingsForMesh(InMesh, nullptr, bIncludeAssets))
 	{
 		MarkPackageDirty();
 
diff --git a/Engine/Source/Runtime/Foliage/Private/ProceduralFoliageInstance.cpp b/Engine/Source/Runtime/Foliage/Private/ProceduralFoliageInstance.cpp
--- a/Engine/Source/Runtime/Foliage/Private/ProceduralFoliageInstance.cpp
+++ b/Engine/Source/Runtime/Foliage/Private/ProceduralFoliageInstance.cpp
@@ -63,13 +63,9 @@ FProceduralFoliageInstance* FProceduralFoliageInstance::Domination(FProceduralFo
 
 	FProceduralFoliageInstance* Dominated = GetLessFit(A, B);
 
-	if (OverlapType == ESimulationOverlap::ShadeOverlap)
+	if (OverlapType == ESimulationOverlap::ShadeOverlap && Dominated->Type->bCanGrowInShade)
 	{
-		const UFoliageType_InstancedStaticMesh* DominatedType = Dominated->Type;
-		if (DominatedType->bGrowsInShade)
-		{
-			return nullptr;
-		}
+		return nullptr;
 	}
 
 	return Dominated;
@@ -97,9 +93,4 @@ float FProceduralFoliageInstance::GetCollisionRadius() const
 void FProceduralFoliageInstance::TerminateInstance()
 {
 	bAlive = false;
-}
-
-bool FProceduralFoliageInstance::GrowsInShade() const 
-{ 
-	return Type->bGrowsInShade; 
-}
+}
\ No newline at end of file
diff --git a/Engine/Source/Runtime/Foliage/Private/ProceduralFoliageTile.cpp b/Engine/Source/Runtime/Foliage/Private/ProceduralFoliageTile.cpp
--- a/Engine/Source/Runtime/Foliage/Private/ProceduralFoliageTile.cpp
+++ b/Engine/Source/Runtime/Foliage/Private/ProceduralFoliageTile.cpp
@@ -130,7 +130,7 @@ void UProceduralFoliageTile::AgeSeeds()
 		if (Instance->IsAlive())
 		{
 			const UFoliageType_InstancedStaticMesh* Type = Instance->Type;
-			if (SimulationStep <= Type->NumSteps && Type->bGrowsInShade == bSimulateShadeGrowth)
+			if (SimulationStep <= Type->NumSteps && Type->GetSpawnsInShade() == bSimulateOnlyInShade)
 			{
 				const float CurrentAge = Instance->Age;
 				const float NewAge = Type->GetNextAge(Instance->Age, 1);
@@ -172,7 +172,7 @@ void UProceduralFoliageTile::SpreadSeeds(TArray<FProceduralFoliageInstance*>& Ne
 
 		const UFoliageType_InstancedStaticMesh* Type = Inst->Type;
 
-		if (SimulationStep <= Type->NumSteps  && Type->bGrowsInShade == bSimulateShadeGrowth)
+		if (SimulationStep <= Type->NumSteps  && Type->GetSpawnsInShade() == bSimulateOnlyInShade)
 		{
 			for (int32 i = 0; i < Type->SeedsPerStep; ++i)
 			{
@@ -210,7 +210,7 @@ void UProceduralFoliageTile::AddRandomSeeds(TArray<FProceduralFoliageInstance*>&
 	{
 		if (UserCancelled()){ return; }
 		const UFoliageType_InstancedStaticMesh* TypeInstance = FoliageTypeObject.GetInstance();
-		if (TypeInstance && TypeInstance->bGrowsInShade == bSimulateShadeGrowth)
+		if (TypeInstance && TypeInstance->GetSpawnsInShade() == bSimulateOnlyInShade)
 		{
 			{	//compute the number of initial seeds
 				const int32 NumSeeds = FMath::RoundToInt(TypeInstance->GetSeedDensitySquared() * SizeTenM2);
@@ -277,7 +277,7 @@ void UProceduralFoliageTile::AddRandomSeeds(TArray<FProceduralFoliageInstance*>&
 			float InitY = 0.f;
 			float NeededRadius = 0.f;
 
-			if (bSimulateShadeGrowth && LastShadeCastingIndex >= 0)
+			if (bSimulateOnlyInShade && LastShadeCastingIndex >= 0)
 			{
 				const int32 InstanceSpawnerIdx = TypeRandomStream.FRandRange(0, LastShadeCastingIndex);
 				const FProceduralFoliageInstance& Spawner = InstancesArray[InstanceSpawnerIdx];
@@ -407,14 +407,14 @@ void UProceduralFoliageTile::StepSimulation()
 	FlushPendingRemovals();
 }
 
-void UProceduralFoliageTile::RunSimulation(const int32 MaxNumSteps, bool bShadeGrowth)
+void UProceduralFoliageTile::RunSimulation(const int32 MaxNumSteps, bool bOnlyInShade)
 {
 	int32 MaxSteps = 0;
 
 	for (const FFoliageTypeObject& FoliageTypeObject : FoliageSpawner->GetFoliageTypes())
 	{
 		const UFoliageType_InstancedStaticMesh* TypeInstance = FoliageTypeObject.GetInstance();
-		if (TypeInstance && TypeInstance->bGrowsInShade == bShadeGrowth)
+		if (TypeInstance && TypeInstance->GetSpawnsInShade() == bOnlyInShade)
 		{
 			MaxSteps = FMath::Max(MaxSteps, TypeInstance->NumSteps + 1);
 		}
@@ -426,7 +426,7 @@ void UProceduralFoliageTile::RunSimulation(const int32 MaxNumSteps, bool bShadeG
 	}
 
 	SimulationStep = 0;
-	bSimulateShadeGrowth = bShadeGrowth;
+	bSimulateOnlyInShade = bOnlyInShade;
 	for (int32 Step = 0; Step < MaxSteps; ++Step)
 	{
 		StepSimulation();
diff --git a/Engine/Source/Runtime/Foliage/Public/FoliageType.h b/Engine/Source/Runtime/Foliage/Public/FoliageType.h
--- a/Engine/Source/Runtime/Foliage/Public/FoliageType.h
+++ b/Engine/Source/Runtime/Foliage/Public/FoliageType.h
@@ -76,6 +76,22 @@ public:
 	UPROPERTY(EditAnywhere, Category=Painting, meta=(UIMin=0, ClampMin=0, ReapplyCondition="ReapplyRadius"))
 	float Radius;
 
+	/** Specifies foliage instance scaling behavior when painting. */
+	UPROPERTY(EditAnywhere, Category=Painting, meta=(ReapplyCondition="ReapplyScaling"))
+	EFoliageScaling Scaling;
+
+	/** Specifies the range of scale, from minimum to maximum, to apply to a foliage instance's X Scale property */
+	UPROPERTY(EditAnywhere, Category=Painting, meta=(ClampMin="0.001", UIMin="0.001", ReapplyCondition="ReapplyScaleX"))
+	FFloatInterval ScaleX;
+
+	/** Specifies the range of scale, from minimum to maximum, to apply to a foliage instance's Y Scale property */
+	UPROPERTY(EditAnywhere, Category=Painting, meta=(ClampMin="0.001", UIMin="0.001", ReapplyCondition="ReapplyScaleY"))
+	FFloatInterval ScaleY;
+
+	/** Specifies the range of scale, from minimum to maximum, to apply to a foliage instance's Z Scale property */
+	UPROPERTY(EditAnywhere, Category=Painting, meta=(ClampMin="0.001", UIMin="0.001", ReapplyCondition="ReapplyScaleZ"))
+	FFloatInterval ScaleZ;
+
 	/** 
 	 *  When painting on static meshes, foliage instance placement can be limited to areas where the static mesh has values in the selected vertex color channel(s). 
 	 *  This allows a static mesh to mask out certain areas to prevent foliage from being placed there
@@ -97,22 +113,6 @@ public:
 public:
 	// PLACEMENT
 
-	/** Specifies foliage instance scaling behavior when painting. */
-	UPROPERTY(EditAnywhere, Category=Placement, meta=(ReapplyCondition="ReapplyScaling"))
-	EFoliageScaling Scaling;
-
-	/** Specifies the range of scale, from minimum to maximum, to apply to a foliage instance's X Scale property */
-	UPROPERTY(EditAnywhere, Category=Placement, meta=(ClampMin="0.001", UIMin="0.001", ReapplyCondition="ReapplyScaleX"))
-	FFloatInterval ScaleX;
-
-	/** Specifies the range of scale, from minimum to maximum, to apply to a foliage instance's Y Scale property */
-	UPROPERTY(EditAnywhere, Category=Placement, meta=(ClampMin="0.001", UIMin="0.001", ReapplyCondition="ReapplyScaleY"))
-	FFloatInterval ScaleY;
-
-	/** Specifies the range of scale, from minimum to maximum, to apply to a foliage instance's Z Scale property */
-	UPROPERTY(EditAnywhere, Category=Placement, meta=(ClampMin="0.001", UIMin="0.001", ReapplyCondition="ReapplyScaleZ"))
-	FFloatInterval ScaleZ;
-
 	/** Specifies a range from minimum to maximum of the offset to apply to a foliage instance's Z location */
 	UPROPERTY(EditAnywhere, Category=Placement, meta=(DisplayName="Z Offset", ReapplyCondition="ReapplyZOffset"))
 	FFloatInterval ZOffset;
@@ -240,6 +240,7 @@ public:
 	FOLIAGE_API float GetScaleForAge(const float Age) const;
 	FOLIAGE_API float GetInitAge(FRandomStream& RandomStream) const;
 	FOLIAGE_API float GetNextAge(const float CurrentAge, const int32 NumSteps) const;
+	FOLIAGE_API bool GetSpawnsInShade() const;
 
 	// COLLISION
 
@@ -283,30 +284,45 @@ public:
 
 	// GROWTH
 
-	/** Whether the species can grow in shade. If this is true shade radius is ignored during overlap tests*/
+	/** If true, seeds of this type will ignore shade radius during overlap tests with other types. */
 	UPROPERTY(Category=Procedural, EditAnywhere, meta=(Subcategory="Growth"))
-	bool bGrowsInShade;
-
-	/** The minimum scale that an instance will be. Corresponds to age = 0. */
-	UPROPERTY(Category=Procedural, EditAnywhere, meta=(Subcategory="Growth", ClampMin="0.0", UIMin="0.0"))
-	float MinScale;
+	bool bCanGrowInShade;
 
-	/** The maximum scale that a seed will grow to. Corresponds to age = 10. */
-	UPROPERTY(Category=Procedural, EditAnywhere, meta=(Subcategory="Growth", ClampMin="0.0", UIMin="0.0"))
-	float MaxScale;
+	/** 
+	 * Whether new seeds are spawned exclusively in shade. Occurs in a second pass after all types that do not spawn in shade have been simulated. 
+	 * Only valid when CanGrowInShade is true. 
+	 */
+	UPROPERTY(Category=Procedural, EditAnywhere, meta=(Subcategory="Growth", EditCondition="bCanGrowInShade"))
+	bool bSpawnsInShade;
 
-	/** Specifies the oldest a new seed can be. The new seed will have an age randomly distributed in [0,InitMaxAge] */
+	/** Allows a new seed to be older than 0 when created. New seeds will be randomly assigned an age in the range [0,MaxInitialAge] */
 	UPROPERTY(Category=Procedural, EditAnywhere, meta=(Subcategory="Growth", ClampMin="0.0", UIMin="0.0"))
-	float InitialMaxAge;
+	float MaxInitialAge;
 
 	/** Specifies the oldest a seed can be. After reaching this age the instance will still spread seeds, but will not get any older*/
 	UPROPERTY(Category=Procedural, EditAnywhere, meta=(Subcategory="Growth", ClampMin="0.0", UIMin="0.0"))
 	float MaxAge;
 
-	/** When two instances overlap we must determine which instance to remove. The instance with a lower OverlapPriority will be removed. In the case where OverlapPriority is the same regular simulation rules apply.*/
+	/** 
+	 * When two instances overlap we must determine which instance to remove. 
+	 * The instance with a lower OverlapPriority will be removed. 
+	 * In the case where OverlapPriority is the same regular simulation rules apply.
+	 */
 	UPROPERTY(Category=Procedural, EditAnywhere, meta=(Subcategory="Growth"))
 	float OverlapPriority;
 
+	/** The scale range of this type when being procedurally generated. Configured with the Scale Curve. */
+	UPROPERTY(Category=Procedural, EditAnywhere, meta = (Subcategory = "Growth", ClampMin = "0.001", UIMin = "0.001"))
+	FFloatInterval ProceduralScale;
+
+	/** 
+	 * Instance scale factor as a function of normalized age (i.e. Current Age / Max Age).
+	 * X = 0 corresponds to Age = 0, X = 1 corresponds to Age = Max Age.
+	 * Y = 0 corresponds to Min Scale, Y = 1 corresponds to Max Scale.
+	 */
+	UPROPERTY(Category = Procedural, EditAnywhere, meta = (Subcategory = "Growth", XAxisName = "Normalized Age", YAxisName = "Scale Factor"))
+	FRuntimeFloatCurve ScaleCurve;
+
 	UPROPERTY()
 	int32 ChangeCount;
 
@@ -375,12 +391,8 @@ public:
 
 private:
 
-	/** The curve used to interpolate the instance scale.*/
-	UPROPERTY(Category = Procedural, EditAnywhere, meta = (Subcategory="Growth"))
-	FRuntimeFloatCurve ScaleCurve;
-	
-	// Deprecated since FFoliageCustomVersion::FoliageTypeCustomization
 #if WITH_EDITORONLY_DATA
+	// Deprecated since FFoliageCustomVersion::FoliageTypeCustomization
 	UPROPERTY()
 	float ScaleMinX_DEPRECATED;
 
@@ -434,6 +446,12 @@ private:
 	
 	UPROPERTY()
 	float MinGroundSlope_DEPRECATED;
+
+	UPROPERTY()
+	float MinScale_DEPRECATED;
+
+	UPROPERTY()
+	float MaxScale_DEPRECATED;
 #endif// WITH_EDITORONLY_DATA
 };
 
diff --git a/Engine/Source/Runtime/Foliage/Public/InstancedFoliageActor.h b/Engine/Source/Runtime/Foliage/Public/InstancedFoliageActor.h
--- a/Engine/Source/Runtime/Foliage/Public/InstancedFoliageActor.h
+++ b/Engine/Source/Runtime/Foliage/Public/InstancedFoliageActor.h
@@ -53,7 +53,7 @@ public:
 
 
 	// Performs a reverse lookup from a mesh to its settings
-	FOLIAGE_API UFoliageType* GetSettingsForMesh(const UStaticMesh* InMesh, FFoliageMeshInfo** OutMeshInfo = nullptr);
+	FOLIAGE_API UFoliageType* GetSettingsForMesh(const UStaticMesh* InMesh, FFoliageMeshInfo** OutMeshInfo = nullptr, bool bIncludeAssets = true);
 
 	// Finds the number of instances overlapping with the sphere. 
 	FOLIAGE_API int32 GetOverlappingSphereCount(const UFoliageType* FoliageType, const FSphere& Sphere) const;
@@ -136,7 +136,7 @@ public:
 	// Add a new static mesh.
 	FOLIAGE_API FFoliageMeshInfo* AddMesh(UStaticMesh* InMesh, UFoliageType** OutSettings = nullptr, const UFoliageType_InstancedStaticMesh* DefaultSettings = nullptr);
 	FOLIAGE_API FFoliageMeshInfo* AddMesh(UFoliageType* InType);
-	FOLIAGE_API FFoliageMeshInfo* UpdateMeshSettings(const UStaticMesh* InMesh, const UFoliageType_InstancedStaticMesh* DefaultSettings);
+	FOLIAGE_API FFoliageMeshInfo* UpdateMeshSettings(const UStaticMesh* InMesh, const UFoliageType_InstancedStaticMesh* DefaultSettings, bool bIncludeAssets = true);
 
 	// Remove the FoliageType from the list, and all its instances.
 	FOLIAGE_API void RemoveFoliageType(UFoliageType** InFoliageType, int32 Num);
diff --git a/Engine/Source/Runtime/Foliage/Public/ProceduralFoliageInstance.h b/Engine/Source/Runtime/Foliage/Public/ProceduralFoliageInstance.h
--- a/Engine/Source/Runtime/Foliage/Public/ProceduralFoliageInstance.h
+++ b/Engine/Source/Runtime/Foliage/Public/ProceduralFoliageInstance.h
@@ -76,8 +76,6 @@ public:
 
 	void TerminateInstance();
 
-	bool GrowsInShade() const;
-
 public:
 	UPROPERTY(Category = ProceduralFoliageInstance, EditAnywhere, BlueprintReadWrite)
 	FVector Location;
diff --git a/Engine/Source/Runtime/Foliage/Public/ProceduralFoliageTile.h b/Engine/Source/Runtime/Foliage/Public/ProceduralFoliageTile.h
--- a/Engine/Source/Runtime/Foliage/Public/ProceduralFoliageTile.h
+++ b/Engine/Source/Runtime/Foliage/Public/ProceduralFoliageTile.h
@@ -93,9 +93,9 @@ private:
 	void Empty();
 
 	/** Run the simulation for the appropriate number of steps. 
-	 * @param bShadeGrowth Whether the simulation should run exclusively within existing shade from other instances
+	 * @param bOnlyInShade Whether the simulation should run exclusively within existing shade from other instances
 	 */
-	void RunSimulation(const int32 MaxNumSteps, bool bShadeGrowth);
+	void RunSimulation(const int32 MaxNumSteps, bool bOnlyInShade);
 
 	/** Advance the simulation by a step (think of a step as a generation) */
 	void StepSimulation();
@@ -133,7 +133,7 @@ private:
 
 	int32 RandomSeed;
 	FRandomStream RandomStream;
-	bool bSimulateShadeGrowth;
+	bool bSimulateOnlyInShade;
 
 	int32 LastCancel;
 
