[UE-11669] - Procedural foliage volumes duplicate properly

diff --git a/Engine/Source/Runtime/Foliage/Private/ProceduralFoliageComponent.cpp b/Engine/Source/Runtime/Foliage/Private/ProceduralFoliageComponent.cpp
--- a/Engine/Source/Runtime/Foliage/Private/ProceduralFoliageComponent.cpp
+++ b/Engine/Source/Runtime/Foliage/Private/ProceduralFoliageComponent.cpp
@@ -209,6 +209,12 @@ bool UProceduralFoliageComponent::ExecuteSimulation(TArray<FDesiredFoliageInstan
 
 }
 
+void UProceduralFoliageComponent::PostEditImport()
+{
+	// The Guid should always be unique
+	ProceduralGuid = FGuid::NewGuid();
+}
+
 bool UProceduralFoliageComponent::GenerateProceduralContent(TArray <FDesiredFoliageInstance>& OutInstances)
 {
 #if WITH_EDITOR
diff --git a/Engine/Source/Runtime/Foliage/Private/ProceduralFoliageVolume.cpp b/Engine/Source/Runtime/Foliage/Private/ProceduralFoliageVolume.cpp
--- a/Engine/Source/Runtime/Foliage/Private/ProceduralFoliageVolume.cpp
+++ b/Engine/Source/Runtime/Foliage/Private/ProceduralFoliageVolume.cpp
@@ -24,6 +24,12 @@ AProceduralFoliageVolume::AProceduralFoliageVolume(const FObjectInitializer& Obj
 
 #if WITH_EDITOR
 
+void AProceduralFoliageVolume::PostEditImport()
+{
+	// Make sure that this is the component's spawning volume
+	ProceduralComponent->SetSpawningVolume(this);
+}
+
 bool AProceduralFoliageVolume::GetReferencedContentObjects(TArray<UObject*>& Objects) const
 {
 	Super::GetReferencedContentObjects(Objects);
diff --git a/Engine/Source/Runtime/Foliage/Public/ProceduralFoliageComponent.h b/Engine/Source/Runtime/Foliage/Public/ProceduralFoliageComponent.h
--- a/Engine/Source/Runtime/Foliage/Public/ProceduralFoliageComponent.h
+++ b/Engine/Source/Runtime/Foliage/Public/ProceduralFoliageComponent.h
@@ -52,6 +52,9 @@ class FOLIAGE_API UProceduralFoliageComponent : public UActorComponent
 	bool bShowDebugTiles;
 #endif
 
+	// UObject interface
+	virtual void PostEditImport() override;
+
 	/** 
 	 * Runs the procedural foliage simulation to generate a list of desired instances to spawn.
 	 * @return True if the simulation succeeded
diff --git a/Engine/Source/Runtime/Foliage/Public/ProceduralFoliageVolume.h b/Engine/Source/Runtime/Foliage/Public/ProceduralFoliageVolume.h
--- a/Engine/Source/Runtime/Foliage/Public/ProceduralFoliageVolume.h
+++ b/Engine/Source/Runtime/Foliage/Public/ProceduralFoliageVolume.h
@@ -15,6 +15,9 @@ class FOLIAGE_API AProceduralFoliageVolume: public AVolume
 	UPROPERTY(Category = ProceduralFoliage, VisibleAnywhere, BlueprintReadOnly)
 	UProceduralFoliageComponent* ProceduralComponent;
 
+	// UObject interface
+	virtual void PostEditImport() override;
+
 #if WITH_EDITOR
 	virtual bool GetReferencedContentObjects(TArray<UObject*>& Objects) const override;
 #endif
