[UE-14155]:
- An existing local foliage type assigned to the same mesh is no longer stomped by proc foliage from a blueprint class
- Added limited support for foliage type blueprints to the foliage palette
- Properties cannot be edited in the palette, but any changes to the BP propagate to the corresponding instances in the world. In general, using FoliageType assets for proc foliage is still recommended.
- No longer possible to assign the same mesh to multiple local foliage types
- Added source asset info to foliage palette item tooltips

- Local foliage types are now transactional
- Updated asset color for proc foliage spawner and foliage type

diff --git a/Engine/Source/Developer/AssetTools/Private/AssetTypeActions/AssetTypeActions_InstancedFoliageSettings.h b/Engine/Source/Developer/AssetTools/Private/AssetTypeActions/AssetTypeActions_InstancedFoliageSettings.h
--- a/Engine/Source/Developer/AssetTools/Private/AssetTypeActions/AssetTypeActions_InstancedFoliageSettings.h
+++ b/Engine/Source/Developer/AssetTools/Private/AssetTypeActions/AssetTypeActions_InstancedFoliageSettings.h
@@ -7,7 +7,7 @@ class FAssetTypeActions_InstancedFoliageSettings : public FAssetTypeActions_Base
 public:
 	// IAssetTypeActions Implementation
 	virtual FText GetName() const override { return NSLOCTEXT("AssetTypeActions", "AssetTypeActions_InstancedFoliageSettings", "Foliage Type"); }
-	virtual FColor GetTypeColor() const override { return FColor(128,192,255); }
+	virtual FColor GetTypeColor() const override { return FColor(12, 173, 12); }
 	virtual UClass* GetSupportedClass() const override;
 	virtual uint32 GetCategories() override { return EAssetTypeCategories::Misc; }
 };
\ No newline at end of file
diff --git a/Engine/Source/Developer/AssetTools/Private/AssetTypeActions/AssetTypeActions_ProceduralFoliageSpawner.h b/Engine/Source/Developer/AssetTools/Private/AssetTypeActions/AssetTypeActions_ProceduralFoliageSpawner.h
--- a/Engine/Source/Developer/AssetTools/Private/AssetTypeActions/AssetTypeActions_ProceduralFoliageSpawner.h
+++ b/Engine/Source/Developer/AssetTools/Private/AssetTypeActions/AssetTypeActions_ProceduralFoliageSpawner.h
@@ -7,7 +7,7 @@ class FAssetTypeActions_ProceduralFoliageSpawner : public FAssetTypeActions_Base
 public:
 	// IAssetTypeActions Implementation
 	virtual FText GetName() const override { return NSLOCTEXT("AssetTypeActions", "AssetTypeActions_ProceduralFoliageSpawner", "Procedural Foliage Spawner"); }
-	virtual FColor GetTypeColor() const override { return FColor(89, 134, 178); }
+	virtual FColor GetTypeColor() const override { return FColor(7, 103, 7); }
 	virtual UClass* GetSupportedClass() const override;
 	virtual uint32 GetCategories() override { return EAssetTypeCategories::Misc; }
 	virtual bool CanFilter() override;
diff --git a/Engine/Source/Editor/FoliageEdit/Private/FoliageEdMode.cpp b/Engine/Source/Editor/FoliageEdit/Private/FoliageEdMode.cpp
--- a/Engine/Source/Editor/FoliageEdit/Private/FoliageEdMode.cpp
+++ b/Engine/Source/Editor/FoliageEdit/Private/FoliageEdMode.cpp
@@ -45,19 +45,31 @@ DEFINE_LOG_CATEGORY_STATIC(LogFoliage, Log, Warning);
 //
 // FFoliageMeshUIInfo
 //
+FFoliageMeshUIInfo::FFoliageMeshUIInfo(UFoliageType* InSettings)
+	: Settings(InSettings)
+	, InstanceCountCurrentLevel(0)
+	, InstanceCountTotal(0)
+{
+}
+
 FText FFoliageMeshUIInfo::GetNameText() const
 {
+	//@todo: this is redundant with FFoliagePaletteItem::DisplayFName, should probably move sorting implementation over to SFoliagePalette
 	const UFoliageType* FoliageType = Settings;
 	if (FoliageType->IsAsset())
 	{
-		return FText::FromString(FoliageType->GetName());
+		return FText::FromName(FoliageType->GetFName());
+	}
+	else if (UBlueprint* Blueprint = Cast<UBlueprint>(FoliageType->GetClass()->ClassGeneratedBy))
+	{
+		return FText::FromName(Blueprint->GetFName());
 	}
 	else
 	{
 		UStaticMesh* Mesh = FoliageType->GetStaticMesh();
 		if (Mesh)
 		{
-			return FText::FromString(Mesh->GetName());
+			return FText::FromName(Mesh->GetFName());
 		}
 		else
 		{
@@ -273,6 +285,9 @@ void FEdModeFoliage::Enter()
 {
 	FEdMode::Enter();
 
+	// register for any objects replaced
+	GEditor->OnObjectsReplaced().AddRaw(this, &FEdModeFoliage::OnObjectsReplaced);
+
 	// Clear any selection in case the instanced foliage actor is selected
 	GEditor->SelectNone(true, true);
 
@@ -332,6 +347,8 @@ void FEdModeFoliage::Exit()
 	FEditorDelegates::NewCurrentLevel.RemoveAll(this);
 	FWorldDelegates::LevelAddedToWorld.RemoveAll(this);
 	FWorldDelegates::LevelRemovedFromWorld.RemoveAll(this);
+
+	GEditor->OnObjectsReplaced().RemoveAll(this);
 	
 	FoliageMeshList.Empty();
 	
@@ -461,6 +478,46 @@ bool FEdModeFoliage::DisallowMouseDeltaTracking() const
 	return bToolActive;
 }
 
+void FEdModeFoliage::OnObjectsReplaced(const TMap<UObject*, UObject*>& ReplacementMap)
+{
+	bool bAnyFoliageTypeReplaced = false;
+
+	UWorld* World = GetWorld();
+	ULevel* CurrentLevel = World->GetCurrentLevel();
+	const int32 NumLevels = World->GetNumLevels();
+
+	// See if any IFA needs to update a foliage type reference
+	for (int32 LevelIdx = 0; LevelIdx < NumLevels; ++LevelIdx)
+	{
+		ULevel* Level = World->GetLevel(LevelIdx);
+		if (Level && Level->bIsVisible)
+		{
+			AInstancedFoliageActor* IFA = AInstancedFoliageActor::GetInstancedFoliageActorForLevel(Level);
+			if (IFA)
+			{
+				for (auto& ReplacementPair : ReplacementMap)
+				{
+					if (UFoliageType* ReplacedFoliageType = Cast<UFoliageType>(ReplacementPair.Key))
+					{
+						TUniqueObj<FFoliageMeshInfo> MeshInfo;
+						if (IFA->FoliageMeshes.RemoveAndCopyValue(ReplacedFoliageType, MeshInfo))
+						{
+							// Re-add the unique mesh info associated with the replaced foliage type
+							IFA->FoliageMeshes.Add(Cast<UFoliageType>(ReplacementPair.Value), MoveTemp(MeshInfo));
+							bAnyFoliageTypeReplaced = true;
+						}
+					}
+				}
+			}
+		}
+	}
+
+	if (bAnyFoliageTypeReplaced)
+	{
+		PopulateFoliageMeshList();
+	}
+}
+
 /** FEdMode: Called once per frame */
 void FEdModeFoliage::Tick(FEditorViewportClient* ViewportClient, float DeltaTime)
 {
@@ -938,40 +995,6 @@ static void SpawnFoliageInstance(UWorld* InWorld, const UFoliageType* Settings,
 	MeshInfo->AddInstance(IFA, FoliageSettings, Instance, BaseComponent);
 }
 
-FFoliageMeshInfo* UpdateMeshSettings(AInstancedFoliageActor* IFA, const FDesiredFoliageInstance& DesiredInstance)
-{
-	FFoliageMeshInfo* MeshInfo = nullptr;
-
-	if (const UFoliageType_InstancedStaticMesh* FoliageType = Cast<UFoliageType_InstancedStaticMesh>(DesiredInstance.FoliageType))
-	{
-		//We have to add the foliage type into the foliage actor.
-		if (FoliageType->IsAsset())
-		{
-			// Add the foliage type asset
-			IFA->AddFoliageType(FoliageType, &MeshInfo);
-		}
-		else if (UStaticMesh* StaticMesh = FoliageType->GetStaticMesh())
-		{
-			if (IFA->GetSettingsForMesh(StaticMesh, nullptr, false) != nullptr)
-			{
-				// @todo: create some unique designation for types added procedurally to prevent them stomping any existing types (or restrict proc foliage to foliage type assets only)
-
-				// A non-asset foliage type using this mesh already exists in the IFA
-				// Only one such type is supported, so the desired instance overrides it
-				MeshInfo = IFA->UpdateMeshSettings(StaticMesh, FoliageType, false);
-				UE_LOG(LogFoliage, Warning, TEXT("Procedurally generated foliage type using Static Mesh \"%s\" conflicts with an existing non-asset foliage type. To prevent clashes, use Foliage Type assets instead."), *StaticMesh->GetName());
-			}
-			else
-			{
-				// No conflict with an existing type, just add normally
-				IFA->AddFoliageType(FoliageType, &MeshInfo);
-			}
-		}
-	}
-
-	return MeshInfo;
-}
-
 void FEdModeFoliage::AddInstancesImp(UWorld* InWorld, const UFoliageType* Settings, const TArray<FDesiredFoliageInstance>& DesiredInstances, const TArray<int32>& ExistingInstanceBuckets, const float Pressure, LandscapeLayerCacheData* LandscapeLayerCachesPtr, const FFoliageUISettings* UISettings)
 {
 	if (DesiredInstances.Num() == 0)
@@ -1003,7 +1026,7 @@ void FEdModeFoliage::AddInstancesImp(UWorld* InWorld, const UFoliageType* Settin
 				if (!UpdatedTypes.Contains(PotentialInst.DesiredInstance.FoliageType))
 				{
 					UpdatedTypes.Add(PotentialInst.DesiredInstance.FoliageType);
-					UpdateMeshSettings(TargetIFA, PotentialInst.DesiredInstance);
+					TargetIFA->AddFoliageType(PotentialInst.DesiredInstance.FoliageType);
 				}
 			}
 		}
@@ -2343,7 +2366,7 @@ void FEdModeFoliage::PopulateFoliageMeshList()
 {
 	FoliageMeshList.Empty();
 	
-	// Collect set of all available Settings
+	// Collect set of all available foliage types
 	UWorld* World = GetWorld();
 	ULevel* CurrentLevel = World->GetCurrentLevel();
 	const int32 NumLevels = World->GetNumLevels();
@@ -2504,7 +2527,7 @@ UFoliageType* FEdModeFoliage::AddFoliageAsset(UObject* InAsset)
 	if (StaticMesh)
 	{	
 		AInstancedFoliageActor* IFA = AInstancedFoliageActor::GetInstancedFoliageActorForCurrentLevel(GetWorld(), true);
-		FoliageType = IFA->GetSettingsForMesh(StaticMesh);
+		FoliageType = IFA->GetLocalFoliageTypeForMesh(StaticMesh);
 		if (!FoliageType)
 		{
 			IFA->AddMesh(StaticMesh, &FoliageType);
diff --git a/Engine/Source/Editor/FoliageEdit/Private/FoliageEdMode.h b/Engine/Source/Editor/FoliageEdit/Private/FoliageEdMode.h
--- a/Engine/Source/Editor/FoliageEdit/Private/FoliageEdMode.h
+++ b/Engine/Source/Editor/FoliageEdit/Private/FoliageEdMode.h
@@ -137,11 +137,7 @@ struct FFoliageMeshUIInfo
 	int32			InstanceCountCurrentLevel;
 	int32			InstanceCountTotal;
 	
-	FFoliageMeshUIInfo(UFoliageType* InSettings)
-		: Settings(InSettings)
-		, InstanceCountCurrentLevel(0)
-		, InstanceCountTotal(0)
-	{}
+	FFoliageMeshUIInfo(UFoliageType* InSettings);
 
 	bool operator == (const FFoliageMeshUIInfo& Other) const
 	{
@@ -291,6 +287,9 @@ public:
 
 	virtual bool DisallowMouseDeltaTracking() const override;
 
+	/** Called when objects are replaced (after a BP compile for instance) */
+	void OnObjectsReplaced(const TMap<UObject*, UObject*>& ReplacementMap);
+
 	/** Forces real-time perspective viewports */
 	void ForceRealTimeViewports(const bool bEnable, const bool bStoreCurrentState);
 
diff --git a/Engine/Source/Editor/FoliageEdit/Private/FoliagePaletteCommands.cpp b/Engine/Source/Editor/FoliageEdit/Private/FoliagePaletteCommands.cpp
--- a/Engine/Source/Editor/FoliageEdit/Private/FoliagePaletteCommands.cpp
+++ b/Engine/Source/Editor/FoliageEdit/Private/FoliagePaletteCommands.cpp
@@ -7,8 +7,8 @@
 
 void FFoliagePaletteCommands::RegisterCommands()
 {
-	UI_COMMAND(ActivateFoliageType, "Activate Selected", "Sets the selected foliage types in the palette as active (i.e. included in brush actions).", EUserInterfaceActionType::Button, FInputChord());
-	UI_COMMAND(DeactivateFoliageType, "Deactivate Selected", "Sets the selected foliage types in the palette as inactive (i.e. excluded in brush actions).", EUserInterfaceActionType::Button, FInputChord());
+	UI_COMMAND(ActivateFoliageType, "Activate", "Sets the selected foliage types in the palette as active (i.e. included in brush actions).", EUserInterfaceActionType::Button, FInputChord());
+	UI_COMMAND(DeactivateFoliageType, "Deactivate", "Sets the selected foliage types in the palette as inactive (i.e. excluded in brush actions).", EUserInterfaceActionType::Button, FInputChord());
 	UI_COMMAND(RemoveFoliageType, "Remove", "Remove this foliage type from the palette. Removes all associated instances as well.", EUserInterfaceActionType::Button, FInputChord(EKeys::Delete));
 	UI_COMMAND(ShowFoliageTypeInCB, "Show in Content Browser", "Show asset in Content Browser.", EUserInterfaceActionType::Button, FInputChord());
 	UI_COMMAND(SelectAllInstances, "Select All Instances", "Select all instances of this foliage type (must be in a selection mode).", EUserInterfaceActionType::Button, FInputChord());
diff --git a/Engine/Source/Editor/FoliageEdit/Private/FoliagePaletteItem.cpp b/Engine/Source/Editor/FoliageEdit/Private/FoliagePaletteItem.cpp
--- a/Engine/Source/Editor/FoliageEdit/Private/FoliagePaletteItem.cpp
+++ b/Engine/Source/Editor/FoliageEdit/Private/FoliagePaletteItem.cpp
@@ -10,6 +10,7 @@
 #include "LevelEditor.h"
 #include "Editor/UnrealEd/Public/AssetThumbnail.h"
 #include "Engine/StaticMesh.h"
+#include "AssetToolsModule.h"
 
 #define LOCTEXT_NAMESPACE "FoliageEd_Mode"
 
@@ -29,17 +30,57 @@ FFoliagePaletteItemModel::FFoliagePaletteItemModel(FFoliageMeshUIInfoPtr InTypeI
 	, FoliagePalette(InFoliagePalette)
 	, FoliageEditMode(InFoliageEditMode)
 {
+	// Determine the display FName
+	if (TypeInfo->Settings->IsAsset())
+	{
+		DisplayFName = TypeInfo->Settings->GetFName();
+	}
+	else if (UBlueprint* FoliageTypeBP = Cast<UBlueprint>(TypeInfo->Settings->GetClass()->ClassGeneratedBy))
+	{
+		DisplayFName = FoliageTypeBP->GetFName();
+	}
+	else
+	{
+		DisplayFName = TypeInfo->Settings->GetStaticMesh()->GetFName();
+	}
+
+	// Create the thumbnail widget
 	FLevelEditorModule& LevelEditorModule = FModuleManager::LoadModuleChecked<FLevelEditorModule>("LevelEditor");
 	TSharedPtr<FAssetThumbnailPool> ThumbnailPool = LevelEditorModule.GetFirstLevelEditor()->GetThumbnailPool();
-	FAssetData AssetData = FAssetData(TypeInfo->Settings->GetStaticMesh());
+	
+	FAssetData AssetData;
+	if (!IsBlueprint())
+	{
+		AssetData = FAssetData(TypeInfo->Settings->GetStaticMesh());
+	}
+	else
+	{
+		AssetData = FAssetData(TypeInfo->Settings->GetClass()->GetDefaultObject<UFoliageType>()->GetStaticMesh());
+	}
+	
 
 	int32 MaxThumbnailSize = FoliagePaletteConstants::ThumbnailSizeRange.Max;
 	TSharedPtr< FAssetThumbnail > Thumbnail = MakeShareable(new FAssetThumbnail(AssetData, MaxThumbnailSize, MaxThumbnailSize, ThumbnailPool));
+	
 	FAssetThumbnailConfig ThumbnailConfig;
+	FAssetToolsModule& AssetToolsModule = FModuleManager::LoadModuleChecked<FAssetToolsModule>(TEXT("AssetTools"));
 	if (TypeInfo->Settings->IsAsset())
 	{
-		//@todo: access the foliage type asset type action and get the color from that (and maybe change it from pale blue - hard to see against the white thumbnail)
-		ThumbnailConfig.AssetTypeColorOverride = FColor(128, 192, 255);
+		// Get the Foliage Type asset color
+		auto FoliageTypeAssetActions =  AssetToolsModule.Get().GetAssetTypeActionsForClass(UFoliageType_InstancedStaticMesh::StaticClass());
+		if (FoliageTypeAssetActions.IsValid())
+		{
+			ThumbnailConfig.AssetTypeColorOverride = FoliageTypeAssetActions.Pin()->GetTypeColor();
+		}
+	}
+	else if (TypeInfo->Settings->GetClass()->ClassGeneratedBy)
+	{
+		// Get the Blueprint asset color
+		auto BlueprintAssetActions = AssetToolsModule.Get().GetAssetTypeActionsForClass(UBlueprint::StaticClass());
+		if (BlueprintAssetActions.IsValid())
+		{
+			ThumbnailConfig.AssetTypeColorOverride = BlueprintAssetActions.Pin()->GetTypeColor();
+		}
 	}
 
 	ThumbnailWidget = Thumbnail->MakeThumbnailWidget(ThumbnailConfig);
@@ -93,12 +134,12 @@ TSharedRef<SToolTip> FFoliagePaletteItemModel::CreateTooltipWidget() const
 					.BorderImage(FEditorStyle::GetBrush("ContentBrowser.TileViewTooltip.ContentBorder"))
 					[
 						SNew(STextBlock)
-						.Text(this, &FFoliagePaletteItemModel::GetFoliageTypeDisplayNameText)
+						.Text(FText::FromName(DisplayFName))
 						.Font(FEditorStyle::GetFontStyle("ContentBrowser.TileViewTooltip.NameFont"))
 						.HighlightText(this, &FFoliagePaletteItemModel::GetPaletteSearchText)
 					]
 				]
-
+				
 				+ SVerticalBox::Slot()
 				.AutoHeight()
 				.Padding(FMargin(0.f, 3.f, 0.f, 0.f))
@@ -134,6 +175,30 @@ TSharedRef<SToolTip> FFoliagePaletteItemModel::CreateTooltipWidget() const
 							SNew(SVerticalBox)
 
 							+ SVerticalBox::Slot()
+							.Padding(0, 1)
+							.AutoHeight()
+							[
+								SNew(SHorizontalBox)
+								
+								+ SHorizontalBox::Slot()
+								.AutoWidth()
+								[
+									SNew(STextBlock)
+									.Text(LOCTEXT("SourceAssetTypeHeading", "Source Asset Type: "))
+									.ColorAndOpacity(FSlateColor::UseSubduedForeground())
+								]
+
+								+ SHorizontalBox::Slot()
+								.AutoWidth()
+								[
+									SNew(STextBlock)
+									.Text(this, &FFoliagePaletteItemModel::GetSourceAssetTypeText)
+								]
+							]
+
+							+ SVerticalBox::Slot()
+							.Padding(0, 1)
+							.AutoHeight()
 							[
 								SNew(SHorizontalBox)
 								
@@ -189,16 +254,9 @@ TSharedRef<SButton> FFoliagePaletteItemModel::CreateSaveAssetButton(TAttribute<E
 		];
 }
 
-FText FFoliagePaletteItemModel::GetFoliageTypeDisplayNameText() const
+FName FFoliagePaletteItemModel::GetDisplayFName() const
 {
-	if (TypeInfo->Settings->IsAsset())
-	{
-		return FText::FromName(TypeInfo->Settings->GetFName());
-	}
-	else
-	{
-		return FText::FromName(TypeInfo->Settings->GetStaticMesh()->GetFName());
-	}
+	return DisplayFName;
 }
 
 FText FFoliagePaletteItemModel::GetPaletteSearchText() const
@@ -218,11 +276,11 @@ FText FFoliagePaletteItemModel::GetInstanceCountText(bool bRounded) const
 	const int32	InstanceCountTotal = TypeInfo->InstanceCountTotal;
 	const int32	InstanceCountCurrentLevel = TypeInfo->InstanceCountCurrentLevel;
 
-	if (InstanceCountCurrentLevel != InstanceCountTotal)
+	if (!bRounded && InstanceCountCurrentLevel != InstanceCountTotal)
 	{
 		return FText::Format(LOCTEXT("InstanceCount_Total", "{0} ({1})"), FText::AsNumber(InstanceCountCurrentLevel), FText::AsNumber(InstanceCountTotal));
 	}
-	else if (bRounded)
+	else if (InstanceCountCurrentLevel >= 1000.f)
 	{
 		// Note: Instance counts greater than 999 billion (unlikely) will not be formatted properly
 		static const FText Suffixes[EInstanceCountMagnitude::Max] =
@@ -240,15 +298,12 @@ FText FFoliagePaletteItemModel::GetInstanceCountText(bool bRounded) const
 			NumThousands++;
 		}
 
-		if (NumThousands > 0)
-		{
-			// Allow 3 significant figures
-			const int32 NumFractionalDigits = DisplayValue >= 100.f ? 0 : DisplayValue >= 10.f ? 1 : 2;
+		// Allow 3 significant figures
+		FNumberFormattingOptions Options;
+		const int32 NumFractionalDigits = DisplayValue >= 100.f ? 0 : DisplayValue >= 10.f ? 1 : 2;
+		Options.SetMaximumFractionalDigits(NumFractionalDigits);
 
-			FNumberFormattingOptions Options;
-			Options.SetMaximumFractionalDigits(NumFractionalDigits);
-			return FText::Format(LOCTEXT("InstanceCount_CurrentLevel", "{0}{1}"), FText::AsNumber(DisplayValue, &Options), Suffixes[NumThousands - 1]);
-		}
+		return FText::Format(LOCTEXT("InstanceCount_CurrentLevel", "{0}{1}"), FText::AsNumber(DisplayValue, &Options), Suffixes[NumThousands - 1]);
 	}
 
 	return FText::AsNumber(InstanceCountCurrentLevel);
@@ -261,6 +316,12 @@ void FFoliagePaletteItemModel::SetTypeActiveInPalette(bool bSetActiveInPalette)
 	{
 		FoliageType->Modify();
 		FoliageType->IsSelected = bSetActiveInPalette;
+
+		if (IsBlueprint())
+		{
+			FoliageType->GetClass()->ClassGeneratedBy->Modify();
+			FoliageType->GetClass()->GetDefaultObject<UFoliageType>()->IsSelected = bSetActiveInPalette;
+		}
 	}
 }
 
@@ -269,6 +330,16 @@ bool FFoliagePaletteItemModel::IsActive() const
 	return TypeInfo->Settings->IsSelected;
 }
 
+bool FFoliagePaletteItemModel::IsBlueprint() const
+{
+	return !TypeInfo->Settings->IsValidLowLevel() || TypeInfo->Settings->GetClass()->ClassGeneratedBy != nullptr;
+}
+
+bool FFoliagePaletteItemModel::IsAsset() const
+{
+	return TypeInfo->Settings->IsAsset();
+}
+
 void FFoliagePaletteItemModel::HandleCheckStateChanged(const ECheckBoxState NewCheckedState, TAttribute<bool> IsItemWidgetSelected)
 {
 	if (!IsItemWidgetSelected.IsSet()) { return; }
@@ -318,6 +389,22 @@ EVisibility FFoliagePaletteItemModel::GetTooltipThumbnailVisibility() const
 	return (FoliagePalette.IsValid() && FoliagePalette.Pin()->IsActiveViewMode(EFoliagePaletteViewMode::Tree)) ? EVisibility::SelfHitTestInvisible : EVisibility::Collapsed;
 }
 
+FText FFoliagePaletteItemModel::GetSourceAssetTypeText() const
+{
+	if (TypeInfo->Settings->IsAsset())
+	{
+		return LOCTEXT("FoliageTypeAsset", "Foliage Type");
+	}
+	else if (UBlueprint* FoliageTypeBP = Cast<UBlueprint>(TypeInfo->Settings->GetClass()->ClassGeneratedBy))
+	{
+		return LOCTEXT("BlueprintClassAsset", "Blueprint Class");
+	}
+	else
+	{
+		return LOCTEXT("StaticMeshAsset", "Static Mesh");
+	}
+}
+
 ////////////////////////////////////////////////
 // SFoliagePaletteItemTile
 ////////////////////////////////////////////////
@@ -396,7 +483,7 @@ void SFoliagePaletteItemTile::Construct(const FArguments& InArgs, TSharedRef<STa
 
 FLinearColor SFoliagePaletteItemTile::GetTileColorAndOpacity() const
 {
-	float Alpha = Model->GetFoliageType()->IsSelected ? 1.f : 0.5f;
+	float Alpha = Model->IsActive() ? 1.f : 0.5f;
 	return FLinearColor(1.f, 1.f, 1.f, Alpha);
 }
 
@@ -407,7 +494,7 @@ EVisibility SFoliagePaletteItemTile::GetCheckBoxVisibility() const
 
 EVisibility SFoliagePaletteItemTile::GetSaveButtonVisibility() const
 {
-	return IsHovered() && CanShowOverlayItems() ? EVisibility::Visible : EVisibility::Collapsed;
+	return IsHovered() && CanShowOverlayItems() && !Model->IsBlueprint() ? EVisibility::Visible : EVisibility::Collapsed;
 }
 
 EVisibility SFoliagePaletteItemTile::GetInstanceCountVisibility() const
@@ -457,7 +544,7 @@ TSharedRef<SWidget> SFoliagePaletteItemRow::GenerateWidgetForColumn(const FName&
 			.AutoWidth()
 			[
 				SNew(STextBlock)
-				.Text(Model.ToSharedRef(), &FFoliagePaletteItemModel::GetFoliageTypeDisplayNameText)
+				.Text(FText::FromName(Model->GetDisplayFName()))
 				.HighlightText(Model.ToSharedRef(), &FFoliagePaletteItemModel::GetPaletteSearchText)
 			];
 	}
@@ -477,10 +564,17 @@ TSharedRef<SWidget> SFoliagePaletteItemRow::GenerateWidgetForColumn(const FName&
 	}
 	else if (ColumnName == FoliagePaletteTreeColumns::ColumnID_Save)
 	{
-		TableRowContent = Model->CreateSaveAssetButton();
+		auto SaveButtonVisibility = TAttribute<EVisibility>::Create(TAttribute<EVisibility>::FGetter::CreateSP(this, &SFoliagePaletteItemRow::GetSaveButtonVisibility));
+
+		TableRowContent = Model->CreateSaveAssetButton(SaveButtonVisibility);
 	}
 
 	return TableRowContent.ToSharedRef();
 }
 
+EVisibility SFoliagePaletteItemRow::GetSaveButtonVisibility() const
+{
+	return !Model->IsBlueprint() ? EVisibility::Visible : EVisibility::Hidden;
+}
+
 #undef LOCTEXT_NAMESPACE
\ No newline at end of file
diff --git a/Engine/Source/Editor/FoliageEdit/Private/FoliagePaletteItem.h b/Engine/Source/Editor/FoliageEdit/Private/FoliagePaletteItem.h
--- a/Engine/Source/Editor/FoliageEdit/Private/FoliagePaletteItem.h
+++ b/Engine/Source/Editor/FoliageEdit/Private/FoliagePaletteItem.h
@@ -44,8 +44,8 @@ public:
 	/** @return The save asset button */
 	TSharedRef<class SButton> CreateSaveAssetButton(TAttribute<EVisibility> InVisibility = TAttribute<EVisibility>());
 
-	/** Gets the name to display of the foliage type object as text */
-	FText GetFoliageTypeDisplayNameText() const;
+	/** Gets the FName version of the displayed name of this item */
+	FName GetDisplayFName() const;
 
 	/** Gets the current search filter text */
 	FText GetPaletteSearchText() const;
@@ -59,6 +59,12 @@ public:
 	/** Gets whether this foliage type is active in the palette */
 	bool IsActive() const;
 
+	/** @return Whether this palette item represents an instance of a foliage type blueprint class */
+	bool IsBlueprint() const;
+
+	/** @return Whether this palette item represents a foliage type asset */
+	bool IsAsset() const;
+
 private:
 	/** Handles the change in activation of the item in the palette */
 	void HandleCheckStateChanged(const ECheckBoxState NewCheckedState, TAttribute<bool> IsItemWidgetSelected);
@@ -78,9 +84,13 @@ private:
 	/** Gets the visibility of the thumbnail in the tooltip */
 	EVisibility GetTooltipThumbnailVisibility() const;
 
+	/** Gets the source asset type text */
+	FText GetSourceAssetTypeText() const;
+
 private:
 	TSharedPtr<SWidget> ThumbnailWidget;
 
+	FName DisplayFName;
 	FFoliageMeshUIInfoPtr TypeInfo;
 	TWeakPtr<SFoliagePalette> FoliagePalette;
 	FEdModeFoliage* FoliageEditMode;
@@ -117,6 +127,9 @@ public:
 	void Construct(const FArguments& InArgs, TSharedRef<STableViewBase> InOwnerTableView, TSharedPtr<FFoliagePaletteItemModel>& InModel);
 	virtual TSharedRef<SWidget> GenerateWidgetForColumn(const FName& ColumnName) override;
 
+private:
+	EVisibility GetSaveButtonVisibility() const;
+
 private:
 	TSharedPtr<FFoliagePaletteItemModel> Model;
 };
\ No newline at end of file
diff --git a/Engine/Source/Editor/FoliageEdit/Private/FoliageType_InstancedStaticMeshPaintingCustomization.cpp b/Engine/Source/Editor/FoliageEdit/Private/FoliageType_InstancedStaticMeshPaintingCustomization.cpp
new file mode 100644
--- /dev/null
+++ b/Engine/Source/Editor/FoliageEdit/Private/FoliageType_InstancedStaticMeshPaintingCustomization.cpp
@@ -0,0 +1,99 @@
+// Copyright 1998-2015 Epic Games, Inc. All Rights Reserved.
+#include "UnrealEd.h"
+#include "FoliageEditModule.h"
+#include "PropertyEditing.h"
+#include "FoliageType_InstancedStaticMesh.h"
+#include "FoliageType_InstancedStaticMeshPaintingCustomization.h"
+#include "FoliageEdMode.h"
+
+#define LOCTEXT_NAMESPACE "FoliageEd_Mode"
+
+DECLARE_DELEGATE_RetVal(EVisibility, FFoliageVisibilityDelegate);
+
+/////////////////////////////////////////////////////
+// FFoliageTypePaintingCustomization 
+TSharedRef<IDetailCustomization> FFoliageType_InstancedStaticMeshPaintingCustomization::MakeInstance(FEdModeFoliage* InFoliageEditMode)
+{
+	auto Instance = MakeShareable(new FFoliageType_InstancedStaticMeshPaintingCustomization(InFoliageEditMode));
+	return Instance;
+}
+
+FFoliageType_InstancedStaticMeshPaintingCustomization::FFoliageType_InstancedStaticMeshPaintingCustomization(FEdModeFoliage* InFoliageEditMode)
+	: FoliageEditMode(InFoliageEditMode)
+{
+}
+
+void FFoliageType_InstancedStaticMeshPaintingCustomization::CustomizeDetails(IDetailLayoutBuilder& DetailLayoutBuilder)
+{
+	TArray<TWeakObjectPtr<UObject>> CustomizedObjects;
+	DetailLayoutBuilder.GetObjectsBeingCustomized(CustomizedObjects);
+
+	// If a non-asset, non-blueprint foliage type (i.e. one that is local to a single InstancedFoliageActor) is being edited, 
+	// we need to restrict the possible meshes that can be assigned
+	bool bCustomizingLocalFoliageType = false;
+	for (TWeakObjectPtr<UObject>& Object : CustomizedObjects)
+	{
+		if (UFoliageType_InstancedStaticMesh* FoliageTypeISM = Cast<UFoliageType_InstancedStaticMesh>(Object.Get()))
+		{
+			if (!FoliageTypeISM->IsAsset() && !FoliageTypeISM->GetClass()->ClassGeneratedBy)
+			{
+				bCustomizingLocalFoliageType = true;
+				break;
+			}
+		}
+	}
+
+	if (bCustomizingLocalFoliageType)
+	{
+		// Cache the meshes that are not available to the type being customized
+		UnavailableMeshNames.Empty();
+		for (auto& TypeInfo : FoliageEditMode->GetFoliageMeshList())
+		{
+			UFoliageType* FoliageType = TypeInfo->Settings;
+			if (!FoliageType->IsAsset() && !FoliageType->GetClass()->ClassGeneratedBy)
+			{
+				UnavailableMeshNames.Add(FoliageType->GetStaticMesh()->GetFName());
+			}
+		}
+
+		// Create a custom entry box for the mesh that filters out meshes that are already assigned to another local foliage type
+		TSharedRef<IPropertyHandle> MeshPropertyHandle = DetailLayoutBuilder.GetProperty(GET_MEMBER_NAME_CHECKED(UFoliageType_InstancedStaticMesh, Mesh));
+		IDetailPropertyRow& PropertyRow = DetailLayoutBuilder.EditCategory("Mesh").AddProperty(MeshPropertyHandle);
+		
+		TSharedPtr<SWidget> NameWidget;
+		TSharedPtr<SWidget> ValueWidget;
+		FDetailWidgetRow	Row;
+		PropertyRow.GetDefaultWidgets(NameWidget, ValueWidget, Row);
+
+		PropertyRow.CustomWidget()
+		.NameContent()
+		.MinDesiredWidth(Row.NameWidget.MinWidth)
+		.MaxDesiredWidth(Row.NameWidget.MaxWidth)
+		[
+			NameWidget.ToSharedRef()
+		]
+		.ValueContent()
+		.MinDesiredWidth(Row.ValueWidget.MinWidth)
+		.MaxDesiredWidth(Row.ValueWidget.MaxWidth)
+		[
+			SNew(SObjectPropertyEntryBox)
+			.AllowedClass(UStaticMesh::StaticClass())
+			.PropertyHandle(MeshPropertyHandle)
+			.ThumbnailPool(DetailLayoutBuilder.GetThumbnailPool())
+			.OnShouldFilterAsset(this, &FFoliageType_InstancedStaticMeshPaintingCustomization::OnShouldFilterAsset)
+		];
+	}
+}
+
+bool FFoliageType_InstancedStaticMeshPaintingCustomization::OnShouldFilterAsset(const FAssetData& AssetData) const
+{
+	if (AssetData.GetClass() == UStaticMesh::StaticClass())
+	{
+		// The unavailable meshes are already referenced and therefore loaded, so we know we won't filter if GetClass() failed
+		return UnavailableMeshNames.Contains(AssetData.AssetName);
+	}
+
+	return false;
+}
+
+#undef LOCTEXT_NAMESPACE
\ No newline at end of file
diff --git a/Engine/Source/Editor/FoliageEdit/Private/FoliageType_InstancedStaticMeshPaintingCustomization.h b/Engine/Source/Editor/FoliageEdit/Private/FoliageType_InstancedStaticMeshPaintingCustomization.h
new file mode 100644
--- /dev/null
+++ b/Engine/Source/Editor/FoliageEdit/Private/FoliageType_InstancedStaticMeshPaintingCustomization.h
@@ -0,0 +1,29 @@
+// Copyright 1998-2015 Epic Games, Inc. All Rights Reserved.
+#pragma once
+
+#include "IDetailCustomization.h"
+
+class FEdModeFoliage;
+
+/////////////////////////////////////////////////////
+// FFoliageTypePaintingCustomization
+class FFoliageType_InstancedStaticMeshPaintingCustomization : public IDetailCustomization
+{
+public:
+	// Makes a new instance of this detail layout class
+	static TSharedRef<IDetailCustomization> MakeInstance(FEdModeFoliage* InFoliageEditMode);
+
+	// IDetailCustomization interface
+	virtual void CustomizeDetails(IDetailLayoutBuilder& DetailLayout) override;
+	// End of IDetailCustomization interface
+
+private:
+	FFoliageType_InstancedStaticMeshPaintingCustomization(FEdModeFoliage* InFoliageEditMode);
+	bool OnShouldFilterAsset(const FAssetData& AssetData) const;
+
+	/** Pointer to the foliage edit mode. */
+	FEdModeFoliage* FoliageEditMode;
+
+	/** The static meshes that are not available to local foliage types (because they are already assigned to a local foliage type) */
+	TArray<FName> UnavailableMeshNames;
+};
diff --git a/Engine/Source/Editor/FoliageEdit/Private/SFoliagePalette.cpp b/Engine/Source/Editor/FoliageEdit/Private/SFoliagePalette.cpp
--- a/Engine/Source/Editor/FoliageEdit/Private/SFoliagePalette.cpp
+++ b/Engine/Source/Editor/FoliageEdit/Private/SFoliagePalette.cpp
@@ -5,6 +5,7 @@
 #include "FoliagePaletteCommands.h"
 #include "SFoliagePalette.h"
 #include "FoliageTypePaintingCustomization.h"
+#include "FoliageType_InstancedStaticMeshPaintingCustomization.h"
 #include "FoliagePaletteItem.h"
 
 #include "LevelEditor.h"
@@ -88,12 +89,67 @@ private:
 	bool bIsDragOn;
 };
 
+////////////////////////////////////////////////
+// SUneditableFoliageTypeWarning
+////////////////////////////////////////////////
+class SUneditableFoliageTypeWarning : public SCompoundWidget
+{
+public:
+	SLATE_BEGIN_ARGS(SUneditableFoliageTypeWarning)
+		: _WarningText()
+		, _OnHyperlinkClicked()
+	{}
+
+	/** The rich text to show in the warning */
+	SLATE_ATTRIBUTE(FText, WarningText)
+
+		/** Called when the hyperlink in the rich text is clicked */
+		SLATE_EVENT(FSlateHyperlinkRun::FOnClick, OnHyperlinkClicked)
+
+		SLATE_END_ARGS()
+
+		/** Constructs the widget */
+		void Construct(const FArguments& InArgs)
+	{
+		ChildSlot
+			[
+				SNew(SBorder)
+				.BorderImage(FEditorStyle::Get().GetBrush("ToolPanel.GroupBorder"))
+				[
+					SNew(SHorizontalBox)
+					+ SHorizontalBox::Slot()
+					.AutoWidth()
+					.HAlign(HAlign_Center)
+					.VAlign(VAlign_Center)
+					.Padding(2)
+					[
+						SNew(SImage)
+						.Image(FEditorStyle::Get().GetBrush("Icons.Warning"))
+					]
+					+ SHorizontalBox::Slot()
+						.VAlign(VAlign_Center)
+						.Padding(2)
+						[
+							SNew(SRichTextBlock)
+							.DecoratorStyleSet(&FEditorStyle::Get())
+							.Justification(ETextJustify::Left)
+							.TextStyle(FEditorStyle::Get(), "DetailsView.BPMessageTextStyle")
+							.Text(InArgs._WarningText)
+							.AutoWrapText(true)
+							+ SRichTextBlock::HyperlinkDecorator(TEXT("HyperlinkDecorator"), InArgs._OnHyperlinkClicked)
+						]
+				]
+			];
+	}
+};
+
 ////////////////////////////////////////////////
 // SFoliagePalette
 ////////////////////////////////////////////////
 void SFoliagePalette::Construct(const FArguments& InArgs)
 {
 	bItemsNeedRebuild = false;
+	bIsUneditableFoliageTypeSelected = false;
 	bIsActiveTimerRegistered = false;
 
 	FoliageEditMode = InArgs._FoliageEdMode;
@@ -112,8 +168,12 @@ void SFoliagePalette::Construct(const FArguments& InArgs)
 	Args.bShowActorLabel = false;
 	DetailsWidget = PropertyModule.CreateDetailView(Args);
 	DetailsWidget->SetVisibility(FoliageEditMode->UISettings.GetShowPaletteItemDetails() ? EVisibility::SelfHitTestInvisible : EVisibility::Collapsed);
+	DetailsWidget->SetIsPropertyEditingEnabledDelegate(FIsPropertyEditingEnabled::CreateSP(this, &SFoliagePalette::GetIsPropertyEditingEnabled));
 
 	// We want to use our own customization for UFoliageType
+	DetailsWidget->RegisterInstancedCustomPropertyLayout(UFoliageType_InstancedStaticMesh::StaticClass(),
+		FOnGetDetailCustomizationInstance::CreateStatic(&FFoliageType_InstancedStaticMeshPaintingCustomization::MakeInstance, FoliageEditMode)
+		);
 	DetailsWidget->RegisterInstancedCustomPropertyLayout(UFoliageType::StaticClass(), 
 		FOnGetDetailCustomizationInstance::CreateStatic(&FFoliageTypePaintingCustomization::MakeInstance, FoliageEditMode)
 		);
@@ -318,7 +378,22 @@ void SFoliagePalette::Construct(const FArguments& InArgs)
 			// Details
 			+SSplitter::Slot()
 			[
-				DetailsWidget.ToSharedRef()
+				SNew(SVerticalBox)
+
+				+ SVerticalBox::Slot()
+				.Padding(0, 2)
+				.AutoHeight()
+				[
+					SNew(SUneditableFoliageTypeWarning)
+					.WarningText(LOCTEXT("CannotEditBlueprintFoliageTypeWarning", "Blueprint foliage types must be edited in the <a id=\"HyperlinkDecorator\" style=\"DetailsView.BPMessageHyperlinkStyle\">Blueprint</>"))
+					.OnHyperlinkClicked(this, &SFoliagePalette::OnEditFoliageTypeBlueprintHyperlinkClicked)
+					.Visibility(this, &SFoliagePalette::GetUneditableFoliageTypeWarningVisibility)
+				]
+
+				+ SVerticalBox::Slot()
+				[
+					DetailsWidget.ToSharedRef()
+				]
 			]
 		]
 	];
@@ -326,6 +401,10 @@ void SFoliagePalette::Construct(const FArguments& InArgs)
 	UpdatePalette(true);
 }
 
+SFoliagePalette::~SFoliagePalette()
+{
+}
+
 FReply SFoliagePalette::OnKeyDown(const FGeometry& MyGeometry, const FKeyEvent& InKeyEvent)
 {
 	if (UICommandList->ProcessCommandBindings(InKeyEvent))
@@ -362,7 +441,9 @@ void SFoliagePalette::UpdateThumbnailForType(UFoliageType* FoliageType)
 				GetActiveViewWidget()->SetItemSelection(Item, true);
 			}
 
-			UpdatePalette();
+			// If a local foliage type changed its mesh, we need to rebuild the palette to ensure a consistent order
+			const bool bRebuild = !Item->IsBlueprint() && !Item->IsAsset();
+			UpdatePalette(bRebuild);
 			break;
 		}
 	}
@@ -370,9 +451,9 @@ void SFoliagePalette::UpdateThumbnailForType(UFoliageType* FoliageType)
 
 bool SFoliagePalette::AnySelectedTileHovered() const
 {
-	for (auto& TypeInfo : GetActiveViewWidget()->GetSelectedItems())
+	for (FFoliagePaletteItemModelPtr& PaletteItem : GetActiveViewWidget()->GetSelectedItems())
 	{
-		TSharedPtr<ITableRow> Tile = TileViewWidget->WidgetFromItem(TypeInfo);
+		TSharedPtr<ITableRow> Tile = TileViewWidget->WidgetFromItem(PaletteItem);
 		if (Tile.IsValid() && Tile->AsWidget()->IsHovered())
 		{
 			return true;
@@ -554,7 +635,7 @@ TSharedRef<SWidgetSwitcher> SFoliagePalette::CreatePaletteViews()
 
 void SFoliagePalette::GetPaletteItemFilterString(FFoliagePaletteItemModelPtr PaletteItemModel, TArray<FString>& OutArray) const
 {
-	OutArray.Add(PaletteItemModel->GetFoliageTypeDisplayNameText().ToString());
+	OutArray.Add(PaletteItemModel->GetDisplayFName().ToString());
 }
 
 void SFoliagePalette::OnSearchTextChanged(const FText& InFilterText)
@@ -639,6 +720,18 @@ FText SFoliagePalette::GetSearchText() const
 void SFoliagePalette::OnSelectionChanged(FFoliagePaletteItemModelPtr Item, ESelectInfo::Type SelectInfo)
 {
 	RefreshDetailsWidget();
+
+	bIsUneditableFoliageTypeSelected = false;
+	for (FFoliagePaletteItemModelPtr& PaletteItem : GetActiveViewWidget()->GetSelectedItems())
+	{
+		// Currently entries from blueprint classes cannot be edited in the palette
+		// as changes do not propagate to the BP class and changes to the BP class stomp any changes made to the instance in the palette item
+		if (PaletteItem->IsBlueprint())
+		{
+			bIsUneditableFoliageTypeSelected = true;
+			break;
+		}
+	}
 }
 
 void SFoliagePalette::OnItemDoubleClicked(FFoliagePaletteItemModelPtr Item) const
@@ -778,7 +871,7 @@ TSharedPtr<SWidget> SFoliagePalette::ConstructFoliageTypeContextMenu()
 	auto SelectedItems = GetActiveViewWidget()->GetSelectedItems();
 	if (SelectedItems.Num() > 0)
 	{
-		const bool bShowSaveAsOption = SelectedItems.Num() == 1 && !SelectedItems[0]->GetFoliageType()->IsAsset();
+		const bool bShowSaveAsOption = SelectedItems.Num() == 1 && !SelectedItems[0]->IsAsset() && !SelectedItems[0]->IsBlueprint();
 		if (bShowSaveAsOption)
 		{
 			MenuBuilder.BeginSection("StaticMeshFoliageTypeOptions", LOCTEXT("StaticMeshFoliageTypeOptionsHeader", "Static Mesh"));
@@ -806,7 +899,9 @@ TSharedPtr<SWidget> SFoliagePalette::ConstructFoliageTypeContextMenu()
 					FSlateIcon(FEditorStyle::GetStyleSetName(), "Level.SaveIcon16x"),
 					FUIAction(
 						FExecuteAction::CreateSP(this, &SFoliagePalette::OnSaveSelected),
-						FCanExecuteAction::CreateSP(this, &SFoliagePalette::OnCanSaveAnySelectedAssets)
+						FCanExecuteAction::CreateSP(this, &SFoliagePalette::OnCanSaveAnySelectedAssets),
+						FIsActionChecked(),
+						FIsActionButtonVisible::CreateSP(this, &SFoliagePalette::GetIsPropertyEditingEnabled)
 						),
 					NAME_None
 					);
@@ -1000,6 +1095,10 @@ void SFoliagePalette::OnShowFoliageTypeInCB()
 		{
 			SelectedAssets.Add(FoliageType);
 		}
+		else if (UBlueprint* FoliageTypeBP = Cast<UBlueprint>(FoliageType->GetClass()->ClassGeneratedBy))
+		{
+			SelectedAssets.Add(FoliageTypeBP);
+		}
 		else
 		{
 			SelectedAssets.Add(FoliageType->GetStaticMesh());
@@ -1160,6 +1259,11 @@ void SFoliagePalette::RefreshDetailsWidget()
 	DetailsWidget->SetObjects(SelectedFoliageTypes, bForceRefresh);
 }
 
+bool SFoliagePalette::GetIsPropertyEditingEnabled() const
+{
+	return !bIsUneditableFoliageTypeSelected;
+}
+
 FText SFoliagePalette::GetDetailsNameAreaText() const
 {
 	FText OutText;
@@ -1167,17 +1271,7 @@ FText SFoliagePalette::GetDetailsNameAreaText() const
 	auto SelectedItems = GetActiveViewWidget()->GetSelectedItems();
 	if (SelectedItems.Num() == 1)
 	{
-		FName DisplayName;
-		UFoliageType* SelectedType = SelectedItems[0]->GetFoliageType();
-		if (SelectedType->IsAsset())
-		{
-			DisplayName = SelectedType->GetFName();
-		}
-		else
-		{
-			DisplayName = SelectedType->GetStaticMesh()->GetFName();
-		}
-		OutText = FText::FromName(DisplayName);
+		OutText = FText::FromName(SelectedItems[0]->GetDisplayFName());
 	}
 	else if (SelectedItems.Num() > 1)
 	{
@@ -1208,6 +1302,32 @@ FReply SFoliagePalette::OnShowHideDetailsClicked() const
 	return FReply::Handled();
 }
 
+EVisibility SFoliagePalette::GetUneditableFoliageTypeWarningVisibility() const
+{
+	return bIsUneditableFoliageTypeSelected ? EVisibility::Visible : EVisibility::Collapsed;
+}
+
+void SFoliagePalette::OnEditFoliageTypeBlueprintHyperlinkClicked(const FSlateHyperlinkRun::FMetadata& Metadata)
+{
+	UBlueprint* Blueprint = nullptr; 
+	
+	// Get the first selected foliage type blueprint
+	for (FFoliagePaletteItemModelPtr& PaletteItem : GetActiveViewWidget()->GetSelectedItems())
+	{
+		Blueprint = Cast<UBlueprint>(PaletteItem->GetFoliageType()->GetClass()->ClassGeneratedBy);
+		if (Blueprint != nullptr)
+		{
+			break;
+		}
+	}
+
+	if (Blueprint)
+	{
+		// Open the blueprint
+		GEditor->EditObject(Blueprint);
+	}
+}
+
 EActiveTimerReturnType SFoliagePalette::UpdatePaletteItems(double InCurrentTime, float InDeltaTime)
 {
 	if (bItemsNeedRebuild)
@@ -1217,7 +1337,7 @@ EActiveTimerReturnType SFoliagePalette::UpdatePaletteItems(double InCurrentTime,
 		// Cache the currently selected items
 		auto ActiveViewWidget = GetActiveViewWidget();
 		TArray<FFoliagePaletteItemModelPtr> PreviouslySelectedItems = ActiveViewWidget->GetSelectedItems();
-		
+
 		ActiveViewWidget->ClearSelection();
 
 		// Rebuild the list of palette items
@@ -1231,10 +1351,10 @@ EActiveTimerReturnType SFoliagePalette::UpdatePaletteItems(double InCurrentTime,
 		// Restore the selection
 		for (auto& PrevSelectedItem : PreviouslySelectedItems)
 		{
-			// See if there's a new item for the previously selected item's foliage type
+			// Select any replacements for previously selected foliage types
 			for (auto& Item : PaletteItems)
 			{
-				if (PrevSelectedItem->GetFoliageType() == Item->GetFoliageType())
+				if (Item->GetDisplayFName() == PrevSelectedItem->GetDisplayFName())
 				{
 					ActiveViewWidget->SetItemSelection(Item, true);
 					break;
diff --git a/Engine/Source/Editor/FoliageEdit/Private/SFoliagePalette.h b/Engine/Source/Editor/FoliageEdit/Private/SFoliagePalette.h
--- a/Engine/Source/Editor/FoliageEdit/Private/SFoliagePalette.h
+++ b/Engine/Source/Editor/FoliageEdit/Private/SFoliagePalette.h
@@ -26,9 +26,12 @@ public:
 	SLATE_END_ARGS()
 
 	void Construct(const FArguments& InArgs);
+	~SFoliagePalette();
+
+	// SWidget interface
 	virtual FReply OnKeyDown(const FGeometry& MyGeometry, const FKeyEvent& InKeyEvent) override;
 
-	/** Updates the foliage palette, optionally doing a full refresh of the items in the palette as well */
+	/** Updates the foliage palette, optionally doing a full rebuild of the items in the palette as well */
 	void UpdatePalette(bool bRebuildItems = false);
 
 	/** Updates the thumbnail for the given foliage type in the palette */
@@ -197,6 +200,9 @@ private:	// DETAILS
 	/** Refreshes the mesh details widget to match the current selection */
 	void RefreshDetailsWidget();
 
+	/** Gets whether property editing is enabled */
+	bool GetIsPropertyEditingEnabled() const;
+
 	/** Gets the text for the details area header */
 	FText GetDetailsNameAreaText() const;
 
@@ -209,6 +215,12 @@ private:	// DETAILS
 	/** Handles the show/hide details button click */
 	FReply OnShowHideDetailsClicked() const;
 
+	/** Gets the visibility of the uneditable blueprint foliage type warning */
+	EVisibility GetUneditableFoliageTypeWarningVisibility() const;
+
+	/** Handles the click for the uneditable blueprint foliage type warning */
+	void OnEditFoliageTypeBlueprintHyperlinkClicked(const FSlateHyperlinkRun::FMetadata& Metadata);
+
 private:
 	/** Active timer handler to update the items in the palette */
 	EActiveTimerReturnType UpdatePaletteItems(double InCurrentTime, float InDeltaTime);
@@ -247,5 +259,6 @@ private:
 	FEdModeFoliage* FoliageEditMode;
 
 	bool bItemsNeedRebuild : 1;
+	bool bIsUneditableFoliageTypeSelected : 1;
 	bool bIsActiveTimerRegistered : 1;
 };
\ No newline at end of file
diff --git a/Engine/Source/Runtime/Foliage/Private/FoliageStatistics.cpp b/Engine/Source/Runtime/Foliage/Private/FoliageStatistics.cpp
--- a/Engine/Source/Runtime/Foliage/Private/FoliageStatistics.cpp
+++ b/Engine/Source/Runtime/Foliage/Private/FoliageStatistics.cpp
@@ -11,7 +11,7 @@ UFoliageStatistics::UFoliageStatistics(const FObjectInitializer& ObjectInitializ
 {
 }
 
-int32 UFoliageStatistics::FoliageOverlappingSphereCount(UObject* WorldContextObject, const UStaticMesh* Mesh, FVector CenterPosition, float Radius)
+int32 UFoliageStatistics::FoliageOverlappingSphereCount(UObject* WorldContextObject, const UStaticMesh* StaticMesh, FVector CenterPosition, float Radius)
 {
 	UWorld* World = GEngine->GetWorldFromContextObject(WorldContextObject);
 	const FSphere Sphere(CenterPosition, Radius);
@@ -23,8 +23,13 @@ int32 UFoliageStatistics::FoliageOverlappingSphereCount(UObject* WorldContextObj
 		AInstancedFoliageActor* IFA = *It;
 		if (!IFA->IsPendingKill())
 		{
-			const UFoliageType* FoliageType = IFA->GetSettingsForMesh(Mesh);
-			Count += IFA->GetOverlappingSphereCount(FoliageType, Sphere);
+			TArray<const UFoliageType*> FoliageTypes;
+			IFA->GetAllFoliageTypesForMesh(StaticMesh, FoliageTypes);
+
+			for (const auto Type : FoliageTypes)
+			{
+				Count += IFA->GetOverlappingSphereCount(Type, Sphere);
+			}
 		}
 	}
 
@@ -42,8 +47,13 @@ int32 UFoliageStatistics::FoliageOverlappingBoxCount(UObject* WorldContextObject
 		AInstancedFoliageActor* IFA = *It;
 		if (!IFA->IsPendingKill())
 		{
-			const UFoliageType* Type = IFA->GetSettingsForMesh(StaticMesh);
-			Count += IFA->GetOverlappingBoxCount(Type, Box);
+			TArray<const UFoliageType*> FoliageTypes;
+			IFA->GetAllFoliageTypesForMesh(StaticMesh, FoliageTypes);
+
+			for (const auto Type : FoliageTypes)
+			{
+				Count += IFA->GetOverlappingBoxCount(Type, Box);
+			}
 		}
 	}
 
@@ -60,8 +70,13 @@ void UFoliageStatistics::FoliageOverlappingBoxTransforms(UObject* WorldContextOb
 		AInstancedFoliageActor* IFA = *It;
 		if (!IFA->IsPendingKill())
 		{
-			const UFoliageType* Type = IFA->GetSettingsForMesh(StaticMesh);
-			IFA->GetOverlappingBoxTransforms(Type, Box, OutTransforms);
+			TArray<const UFoliageType*> FoliageTypes;
+			IFA->GetAllFoliageTypesForMesh(StaticMesh, FoliageTypes);
+
+			for (const auto Type : FoliageTypes)
+			{
+				IFA->GetOverlappingBoxTransforms(Type, Box, OutTransforms);
+			}
 		}
 	}
 }
diff --git a/Engine/Source/Runtime/Foliage/Private/InstancedFoliage.cpp b/Engine/Source/Runtime/Foliage/Private/InstancedFoliage.cpp
--- a/Engine/Source/Runtime/Foliage/Private/InstancedFoliage.cpp
+++ b/Engine/Source/Runtime/Foliage/Private/InstancedFoliage.cpp
@@ -507,7 +507,7 @@ void UFoliageType::PostEditChangeProperty(struct FPropertyChangedEvent& Property
 
 	//@todo: move this into FoliageType_InstancedStaticMesh
 	// Check to see if the mesh is what changed
-	const bool bMeshChanged = PropertyChangedEvent.Property && PropertyChangedEvent.Property->GetFName() == "Mesh";
+	const bool bMeshChanged = PropertyChangedEvent.Property && PropertyChangedEvent.Property->GetFName() == GET_MEMBER_NAME_CHECKED(UFoliageType_InstancedStaticMesh, Mesh);
 
 	// Notify any currently-loaded InstancedFoliageActors
 	if (IsFoliageReallocationRequiredForPropertyChange(PropertyChangedEvent))
@@ -612,88 +612,96 @@ void FFoliageMeshInfo::UpdateComponentSettings(const UFoliageType* InSettings)
 		bool bNeedsMarkRenderStateDirty = false;
 		bool bNeedsInvalidateLightingCache = false;
 		
-		if (Component->StaticMesh != InSettings->GetStaticMesh())
+		const UFoliageType* FoliageType = InSettings;
+		if (InSettings->GetClass()->ClassGeneratedBy)
 		{
-			Component->StaticMesh = InSettings->GetStaticMesh();
+			// If we're updating settings for a BP foliage type, use the CDO
+			FoliageType = InSettings->GetClass()->GetDefaultObject<UFoliageType>();
+		}
+
+		if (Component->StaticMesh != FoliageType->GetStaticMesh())
+		{
+			Component->StaticMesh = FoliageType->GetStaticMesh();
+
 			bNeedsInvalidateLightingCache = true;
 			bNeedsMarkRenderStateDirty = true;
 		}
 
-		EComponentMobility::Type NewMobility = InSettings->bEnableStaticLighting ? EComponentMobility::Static : EComponentMobility::Movable;
+		EComponentMobility::Type NewMobility = FoliageType->bEnableStaticLighting ? EComponentMobility::Static : EComponentMobility::Movable;
 		if (Component->Mobility != NewMobility)
 		{
 			Component->Mobility = NewMobility;
 			bNeedsMarkRenderStateDirty = true;
 			bNeedsInvalidateLightingCache = true;
 		}
-		if (Component->InstanceStartCullDistance != InSettings->CullDistance.Min)
+		if (Component->InstanceStartCullDistance != FoliageType->CullDistance.Min)
 		{
-			Component->InstanceStartCullDistance = InSettings->CullDistance.Min;
+			Component->InstanceStartCullDistance = FoliageType->CullDistance.Min;
 			bNeedsMarkRenderStateDirty = true;
 		}
-		if (Component->InstanceEndCullDistance != InSettings->CullDistance.Max)
+		if (Component->InstanceEndCullDistance != FoliageType->CullDistance.Max)
 		{
-			Component->InstanceEndCullDistance = InSettings->CullDistance.Max;
+			Component->InstanceEndCullDistance = FoliageType->CullDistance.Max;
 			bNeedsMarkRenderStateDirty = true;
 		}
-		if (Component->CastShadow != InSettings->CastShadow)
+		if (Component->CastShadow != FoliageType->CastShadow)
 		{
-			Component->CastShadow = InSettings->CastShadow;
+			Component->CastShadow = FoliageType->CastShadow;
 			bNeedsMarkRenderStateDirty = true;
 			bNeedsInvalidateLightingCache = true;
 		}
-		if (Component->bCastDynamicShadow != InSettings->bCastDynamicShadow)
+		if (Component->bCastDynamicShadow != FoliageType->bCastDynamicShadow)
 		{
-			Component->bCastDynamicShadow = InSettings->bCastDynamicShadow;
+			Component->bCastDynamicShadow = FoliageType->bCastDynamicShadow;
 			bNeedsMarkRenderStateDirty = true;
 			bNeedsInvalidateLightingCache = true;
 		}
-		if (Component->bCastStaticShadow != InSettings->bCastStaticShadow)
+		if (Component->bCastStaticShadow != FoliageType->bCastStaticShadow)
 		{
-			Component->bCastStaticShadow = InSettings->bCastStaticShadow;
+			Component->bCastStaticShadow = FoliageType->bCastStaticShadow;
 			bNeedsMarkRenderStateDirty = true;
 			bNeedsInvalidateLightingCache = true;
 		}
-		if (Component->bAffectDynamicIndirectLighting != InSettings->bAffectDynamicIndirectLighting)
+		if (Component->bAffectDynamicIndirectLighting != FoliageType->bAffectDynamicIndirectLighting)
 		{
-			Component->bAffectDynamicIndirectLighting = InSettings->bAffectDynamicIndirectLighting;
+			Component->bAffectDynamicIndirectLighting = FoliageType->bAffectDynamicIndirectLighting;
 			bNeedsMarkRenderStateDirty = true;
 			bNeedsInvalidateLightingCache = true;
 		}
-		if (Component->bAffectDistanceFieldLighting != InSettings->bAffectDistanceFieldLighting)
+		if (Component->bAffectDistanceFieldLighting != FoliageType->bAffectDistanceFieldLighting)
 		{
-			Component->bAffectDistanceFieldLighting = InSettings->bAffectDistanceFieldLighting;
+			Component->bAffectDistanceFieldLighting = FoliageType->bAffectDistanceFieldLighting;
 			bNeedsMarkRenderStateDirty = true;
 			bNeedsInvalidateLightingCache = true;
 		}
-		if (Component->bCastShadowAsTwoSided != InSettings->bCastShadowAsTwoSided)
+		if (Component->bCastShadowAsTwoSided != FoliageType->bCastShadowAsTwoSided)
 		{
-			Component->bCastShadowAsTwoSided = InSettings->bCastShadowAsTwoSided;
+			Component->bCastShadowAsTwoSided = FoliageType->bCastShadowAsTwoSided;
 			bNeedsMarkRenderStateDirty = true;
 			bNeedsInvalidateLightingCache = true;
 		}
-		if (Component->bReceivesDecals != InSettings->bReceivesDecals)
+		if (Component->bReceivesDecals != FoliageType->bReceivesDecals)
 		{
-			Component->bReceivesDecals = InSettings->bReceivesDecals;
+			Component->bReceivesDecals = FoliageType->bReceivesDecals;
 			bNeedsMarkRenderStateDirty = true;
 			bNeedsInvalidateLightingCache = true;
 		}
-		if (Component->bOverrideLightMapRes != InSettings->bOverrideLightMapRes)
+		if (Component->bOverrideLightMapRes != FoliageType->bOverrideLightMapRes)
 		{
-			Component->bOverrideLightMapRes = InSettings->bOverrideLightMapRes;
+			Component->bOverrideLightMapRes = FoliageType->bOverrideLightMapRes;
 			bNeedsMarkRenderStateDirty = true;
 			bNeedsInvalidateLightingCache = true;
 		}
-		if (Component->OverriddenLightMapRes != InSettings->OverriddenLightMapRes)
+		if (Component->OverriddenLightMapRes != FoliageType->OverriddenLightMapRes)
 		{
-			Component->OverriddenLightMapRes = InSettings->OverriddenLightMapRes;
+			Component->OverriddenLightMapRes = FoliageType->OverriddenLightMapRes;
 			bNeedsMarkRenderStateDirty = true;
 			bNeedsInvalidateLightingCache = true;
 		}
 		
-		Component->BodyInstance.CopyBodyInstancePropertiesFrom(&InSettings->BodyInstance);
+		Component->BodyInstance.CopyBodyInstancePropertiesFrom(&FoliageType->BodyInstance);
 
-		Component->SetCustomNavigableGeometry(InSettings->CustomNavigableGeometry);
+		Component->SetCustomNavigableGeometry(FoliageType->CustomNavigableGeometry);
 
 		if (bNeedsInvalidateLightingCache)
 		{
@@ -1220,7 +1228,7 @@ void AInstancedFoliageActor::GetOverlappingMeshCounts(const FSphere& Sphere, TMa
 
 
 
-UFoliageType* AInstancedFoliageActor::GetSettingsForMesh(const UStaticMesh* InMesh, FFoliageMeshInfo** OutMeshInfo, bool bIncludeAssets)
+UFoliageType* AInstancedFoliageActor::GetLocalFoliageTypeForMesh(const UStaticMesh* InMesh, FFoliageMeshInfo** OutMeshInfo)
 {
 	UFoliageType* ReturnType = nullptr;
 	FFoliageMeshInfo* MeshInfo = nullptr;
@@ -1228,7 +1236,8 @@ UFoliageType* AInstancedFoliageActor::GetSettingsForMesh(const UStaticMesh* InMe
 	for (auto& MeshPair : FoliageMeshes)
 	{
 		UFoliageType* FoliageType = MeshPair.Key;
-		if (FoliageType && (bIncludeAssets || !FoliageType->IsAsset()) && FoliageType->GetStaticMesh() == InMesh)
+		// Check that the type is neither an asset nor blueprint instance
+		if (FoliageType && FoliageType->GetStaticMesh() == InMesh && !FoliageType->IsAsset() && !FoliageType->GetClass()->ClassGeneratedBy)
 		{
 			ReturnType = FoliageType;
 			MeshInfo = &*MeshPair.Value;
@@ -1244,6 +1253,35 @@ UFoliageType* AInstancedFoliageActor::GetSettingsForMesh(const UStaticMesh* InMe
 	return ReturnType;
 }
 
+void AInstancedFoliageActor::GetAllFoliageTypesForMesh(const UStaticMesh* InMesh, TArray<const UFoliageType*>& OutFoliageTypes)
+{
+	for (auto& MeshPair : FoliageMeshes)
+	{
+		UFoliageType* FoliageType = MeshPair.Key;
+		if (FoliageType && FoliageType->GetStaticMesh() == InMesh)
+		{
+			OutFoliageTypes.Add(FoliageType);
+		}
+	}
+}
+
+
+FFoliageMeshInfo* AInstancedFoliageActor::FindFoliageTypeOfClass(TSubclassOf<UFoliageType_InstancedStaticMesh> Class)
+{
+	FFoliageMeshInfo* MeshInfo = nullptr;
+
+	for (auto& MeshPair : FoliageMeshes)
+	{
+		UFoliageType* FoliageType = MeshPair.Key;
+		if (FoliageType && FoliageType->GetClass() == Class)
+		{
+			MeshInfo = &MeshPair.Value.Get();
+			break;
+		}
+	}
+
+	return MeshInfo;
+}
 
 FFoliageMeshInfo* AInstancedFoliageActor::FindMesh(const UFoliageType* InType)
 {
@@ -1629,7 +1667,7 @@ void UpdateSettingsBounds(const UStaticMesh* InMesh, UFoliageType_InstancedStati
 #if WITH_EDITORONLY_DATA
 FFoliageMeshInfo* AInstancedFoliageActor::UpdateMeshSettings(const UStaticMesh* InMesh, const UFoliageType_InstancedStaticMesh* DefaultSettings, bool bIncludeAssets)
 {
-	if (UFoliageType* OldSettings = GetSettingsForMesh(InMesh, nullptr, bIncludeAssets))
+	if (UFoliageType* OldSettings = GetLocalFoliageTypeForMesh(InMesh, nullptr))
 	{
 		MarkPackageDirty();
 
@@ -1654,22 +1692,36 @@ UFoliageType* AInstancedFoliageActor::AddFoliageType(const UFoliageType* InType,
 	
 	if (FoliageType->GetOuter() == this || FoliageType->IsAsset())
 	{
-		auto* MeshInfoPtr = FoliageMeshes.Find(FoliageType);
-		if (!MeshInfoPtr)
+		auto ExistingMeshInfo = FoliageMeshes.Find(FoliageType);
+		if (!ExistingMeshInfo)
+		{
+			Modify();
+			MeshInfo = &FoliageMeshes.Add(FoliageType).Get();
+		}
+		else
+		{
+			MeshInfo = &ExistingMeshInfo->Get();
+		}
+	}
+	else if (FoliageType->GetClass()->ClassGeneratedBy)
+	{
+		// Foliage type blueprint
+		FFoliageMeshInfo* ExistingMeshInfo = FindFoliageTypeOfClass(FoliageType->GetClass());
+		if (!ExistingMeshInfo)
 		{
 			Modify();
 			MeshInfo = &FoliageMeshes.Add(FoliageType).Get();
 		}
 		else
 		{
-			MeshInfo = &MeshInfoPtr->Get();
+			MeshInfo = ExistingMeshInfo;
 		}
 	}
 	else
 	{
 		// Unique meshes only
 		// Multiple entries for same static mesh can be added using FoliageType as an asset
-		FoliageType = GetSettingsForMesh(FoliageType->GetStaticMesh(), &MeshInfo);
+		FoliageType = GetLocalFoliageTypeForMesh(FoliageType->GetStaticMesh(), &MeshInfo);
 		if (FoliageType == nullptr)
 		{
 			Modify();
@@ -1688,7 +1740,7 @@ UFoliageType* AInstancedFoliageActor::AddFoliageType(const UFoliageType* InType,
 
 FFoliageMeshInfo* AInstancedFoliageActor::AddMesh(UStaticMesh* InMesh, UFoliageType** OutSettings, const UFoliageType_InstancedStaticMesh* DefaultSettings)
 {
-	check(GetSettingsForMesh(InMesh) == nullptr);
+	check(GetLocalFoliageTypeForMesh(InMesh) == nullptr);
 
 	MarkPackageDirty();
 
@@ -1704,6 +1756,7 @@ FFoliageMeshInfo* AInstancedFoliageActor::AddMesh(UStaticMesh* InMesh, UFoliageT
 	{
 		Settings = NewObject<UFoliageType_InstancedStaticMesh>(this);
 	}
+	Settings->SetFlags(RF_Transactional);
 	Settings->Mesh = InMesh;
 
 	FFoliageMeshInfo* MeshInfo = AddMesh(Settings);
diff --git a/Engine/Source/Runtime/Foliage/Public/FoliageType_InstancedStaticMesh.h b/Engine/Source/Runtime/Foliage/Public/FoliageType_InstancedStaticMesh.h
--- a/Engine/Source/Runtime/Foliage/Public/FoliageType_InstancedStaticMesh.h
+++ b/Engine/Source/Runtime/Foliage/Public/FoliageType_InstancedStaticMesh.h
@@ -10,7 +10,7 @@ class UFoliageType_InstancedStaticMesh : public UFoliageType
 {
 	GENERATED_UCLASS_BODY()
 
-	UPROPERTY(EditAnywhere, Category=Mesh)
+	UPROPERTY(EditAnywhere, Category=Mesh, meta=(DisplayThumbnail="true"))
 	UStaticMesh* Mesh;
 
 	virtual UStaticMesh* GetStaticMesh() const override
diff --git a/Engine/Source/Runtime/Foliage/Public/InstancedFoliageActor.h b/Engine/Source/Runtime/Foliage/Public/InstancedFoliageActor.h
--- a/Engine/Source/Runtime/Foliage/Public/InstancedFoliageActor.h
+++ b/Engine/Source/Runtime/Foliage/Public/InstancedFoliageActor.h
@@ -52,8 +52,13 @@ public:
 	// End AActor interface.
 
 
-	// Performs a reverse lookup from a mesh to its settings
-	FOLIAGE_API UFoliageType* GetSettingsForMesh(const UStaticMesh* InMesh, FFoliageMeshInfo** OutMeshInfo = nullptr, bool bIncludeAssets = true);
+	// Performs a reverse lookup from a mesh to a local foliage type (i.e. the foliage type owned exclusively by this IFA)
+	FOLIAGE_API UFoliageType* GetLocalFoliageTypeForMesh(const UStaticMesh* InMesh, FFoliageMeshInfo** OutMeshInfo = nullptr);
+	
+	// Performs a reverse lookup from a mesh to all the foliage types that are currently using that mesh (includes assets and blueprint classes)
+	FOLIAGE_API void GetAllFoliageTypesForMesh(const UStaticMesh* InMesh, TArray<const UFoliageType*>& OutFoliageTypes);
+	
+	FOLIAGE_API FFoliageMeshInfo* FindFoliageTypeOfClass(TSubclassOf<UFoliageType_InstancedStaticMesh> Class);
 
 	// Finds the number of instances overlapping with the sphere. 
 	FOLIAGE_API int32 GetOverlappingSphereCount(const UFoliageType* FoliageType, const FSphere& Sphere) const;
