removed unused DisplayOrder property from FoliageType

diff --git a/Engine/Source/Runtime/Foliage/Private/InstancedFoliage.cpp b/Engine/Source/Runtime/Foliage/Private/InstancedFoliage.cpp
--- a/Engine/Source/Runtime/Foliage/Private/InstancedFoliage.cpp
+++ b/Engine/Source/Runtime/Foliage/Private/InstancedFoliage.cpp
@@ -281,7 +281,6 @@ UFoliageType::UFoliageType(const FObjectInitializer& ObjectInitializer)
 	CullDistance.Min = 0;
 	CullDistance.Max = 0;
 	MinimumLayerWeight = 0.5f;
-	DisplayOrder = 0;
 	IsSelected = false;
 	DensityAdjustmentFactor = 1.0f;
 	CollisionWithWorld = false;
@@ -1776,19 +1775,6 @@ FFoliageMeshInfo* AInstancedFoliageActor::AddMesh(UFoliageType* InType)
 
 	Modify();
 
-	if (InType->DisplayOrder == 0)
-	{
-		int32 MaxDisplayOrder = 0;
-		for (auto& MeshPair : FoliageMeshes)
-		{
-			if (MeshPair.Key->DisplayOrder > MaxDisplayOrder)
-			{
-				MaxDisplayOrder = MeshPair.Key->DisplayOrder;
-			}
-		}
-		InType->DisplayOrder = MaxDisplayOrder + 1;
-	}
-
 	FFoliageMeshInfo* MeshInfo = &*FoliageMeshes.Add(InType);
 	MeshInfo->FoliageTypeUpdateGuid = InType->UpdateGuid;
 	InType->IsSelected = true;
diff --git a/Engine/Source/Runtime/Foliage/Public/FoliageType.h b/Engine/Source/Runtime/Foliage/Public/FoliageType.h
--- a/Engine/Source/Runtime/Foliage/Public/FoliageType.h
+++ b/Engine/Source/Runtime/Foliage/Public/FoliageType.h
@@ -226,9 +226,6 @@ public:
 	UPROPERTY(EditAnywhere, Category=InstanceSettings, meta=(HideObjectType=true))
 	TEnumAsByte<EHasCustomNavigableGeometry::Type> CustomNavigableGeometry;
 
-	UPROPERTY()
-	int32 DisplayOrder;
-
 	UPROPERTY()
 	uint32 IsSelected:1;
 
