Fix crash when deleting static mesh of a foliage instance that's embedded in the FLA

UE-15371

diff --git a/Engine/Source/Editor/FoliageEdit/Private/FoliageEdMode.cpp b/Engine/Source/Editor/FoliageEdit/Private/FoliageEdMode.cpp
--- a/Engine/Source/Editor/FoliageEdit/Private/FoliageEdMode.cpp
+++ b/Engine/Source/Editor/FoliageEdit/Private/FoliageEdMode.cpp
@@ -2302,7 +2302,13 @@ void FEdModeFoliage::HandleOnActorSpawned(AActor* Actor)
 
 void FEdModeFoliage::HandleOnFoliageTypeMeshChanged(UFoliageType* FoliageType)
 {
-	StaticCastSharedPtr<FFoliageEdModeToolkit>(Toolkit)->NotifyFoliageTypeMeshChanged(FoliageType);
+	if(FoliageType->IsNotAssetOrBlueprint() && FoliageType->GetStaticMesh() == nullptr)
+	{
+		RemoveFoliageType(&FoliageType, 1);
+	}else
+	{
+		StaticCastSharedPtr<FFoliageEdModeToolkit>(Toolkit)->NotifyFoliageTypeMeshChanged(FoliageType);
+	}
 }
 
 bool FEdModeFoliage::SnapInstanceToGround(AInstancedFoliageActor* InIFA, float AlignMaxAngle, FFoliageMeshInfo& Mesh, int32 InstanceIdx)
diff --git a/Engine/Source/Runtime/Foliage/Private/InstancedFoliage.cpp b/Engine/Source/Runtime/Foliage/Private/InstancedFoliage.cpp
--- a/Engine/Source/Runtime/Foliage/Private/InstancedFoliage.cpp
+++ b/Engine/Source/Runtime/Foliage/Private/InstancedFoliage.cpp
@@ -402,6 +402,11 @@ void UFoliageType::Serialize(FArchive& Ar)
 #endif// WITH_EDITORONLY_DATA
 }
 
+bool UFoliageType::IsNotAssetOrBlueprint() const
+{
+	return IsAsset() == false && Cast<UBlueprint>(GetClass()->ClassGeneratedBy) == nullptr;
+}
+
 FVector UFoliageType::GetRandomScale() const
 {
 	FVector Result(1.0f);
@@ -2288,6 +2293,13 @@ void AInstancedFoliageActor::PostLoad()
 					MeshInfo.Component->ClearFlags(RF_Transactional);
 				}
 			}
+
+			//Clean up case where embeded instances had their static mesh deleted
+			if (FoliageType->IsNotAssetOrBlueprint() && FoliageType->GetStaticMesh() == nullptr)
+			{
+				OnFoliageTypeMeshChangedEvent.Broadcast(FoliageType);
+				RemoveFoliageType(&FoliageType, 1);
+			}
 		}
 
 		// Clean up dead cross-level references
@@ -2311,6 +2323,11 @@ void AInstancedFoliageActor::NotifyFoliageTypeChanged(UFoliageType* FoliageType,
 		{
 			// If the type's mesh has changed, the UI needs to be notified so it can update thumbnails accordingly
 			OnFoliageTypeMeshChangedEvent.Broadcast(FoliageType);
+
+			if(FoliageType->IsNotAssetOrBlueprint() && FoliageType->GetStaticMesh() == nullptr) //If the mesh has been deleted and we're a per foliage actor instance we must remove all instances of the mesh
+			{
+				RemoveFoliageType(&FoliageType, 1);
+			}
 		}
 	}
 }
diff --git a/Engine/Source/Runtime/Foliage/Public/FoliageType.h b/Engine/Source/Runtime/Foliage/Public/FoliageType.h
--- a/Engine/Source/Runtime/Foliage/Public/FoliageType.h
+++ b/Engine/Source/Runtime/Foliage/Public/FoliageType.h
@@ -46,6 +46,8 @@ class UFoliageType : public UObject
 
 	virtual void Serialize(FArchive& Ar) override;
 
+	virtual bool IsNotAssetOrBlueprint() const;
+
 	FOLIAGE_API FVector GetRandomScale() const;
 
 #if WITH_EDITOR
