Fix procedural foliage placement issues when sweep was initially penetrating. The fix is to move the start position up a bit so that the initial penetration never happens.

UE-15431

diff --git a/Engine/Source/Runtime/Foliage/Private/InstancedFoliage.cpp b/Engine/Source/Runtime/Foliage/Private/InstancedFoliage.cpp
--- a/Engine/Source/Runtime/Foliage/Private/InstancedFoliage.cpp
+++ b/Engine/Source/Runtime/Foliage/Private/InstancedFoliage.cpp
@@ -2420,7 +2420,11 @@ bool AInstancedFoliageActor::FoliageTrace(const UWorld* InWorld, FHitResult& Out
 	FCollisionQueryParams QueryParams(InTraceTag, true);
 	QueryParams.bReturnFaceIndex = InbReturnFaceIndex;
 
-	FVector StartTrace = DesiredInstance.StartTrace;
+	//It's possible that with the radius of the shape we will end up with an initial overlap which would place the instance at the top of the procedural volume.
+	//Moving the start trace back a bit will fix this, but it introduces the potential for spawning instances a bit above the volume. This second issue is already somewhat broken because of how sweeps work so it's not too bad, also this is a less common case.
+	//The proper fix would be to do something like EncroachmentCheck where we first do a sweep, then we fix it up if it's overlapping, then check the filters. This is more expensive and error prone so for now we just move the trace up a bit.
+	const FVector Dir = (DesiredInstance.EndTrace - DesiredInstance.StartTrace).GetSafeNormal();
+	const FVector StartTrace = DesiredInstance.StartTrace - (Dir * DesiredInstance.TraceRadius);
 
 	TArray<FHitResult> Hits;
 
@@ -2429,7 +2433,6 @@ bool AInstancedFoliageActor::FoliageTrace(const UWorld* InWorld, FHitResult& Out
 	SphereShape.SetSphere(DesiredInstance.TraceRadius);
 	InWorld->SweepMultiByObjectType(Hits, StartTrace, DesiredInstance.EndTrace, FQuat::Identity, FCollisionObjectQueryParams(ECC_WorldStatic), SphereShape, QueryParams);
 
-
 	for (const FHitResult& Hit : Hits)
 	{
 		if (DesiredInstance.PlacementMode == EFoliagePlacementMode::Procedural)
