Duplicate blueprint case of procedural foliage to fix GLEO issues.

UE-15454

diff --git a/Engine/Source/Runtime/Foliage/Private/InstancedFoliage.cpp b/Engine/Source/Runtime/Foliage/Private/InstancedFoliage.cpp
--- a/Engine/Source/Runtime/Foliage/Private/InstancedFoliage.cpp
+++ b/Engine/Source/Runtime/Foliage/Private/InstancedFoliage.cpp
@@ -1714,6 +1714,7 @@ UFoliageType* AInstancedFoliageActor::AddFoliageType(const UFoliageType* InType,
 		if (!ExistingMeshInfo)
 		{
 			Modify();
+			FoliageType = DuplicateObject<UFoliageType>(InType, this);
 			MeshInfo = &FoliageMeshes.Add(FoliageType).Get();
 		}
 		else
