Blueprint support for foliage InstancedStaticMeshComponent UE-9078
Override mesh material in FoliageType UE-17375
Thumbnails for FoliageType assets

diff --git a/Engine/Source/Editor/FoliageEdit/Private/FoliageEditModule.cpp b/Engine/Source/Editor/FoliageEdit/Private/FoliageEditModule.cpp
--- a/Engine/Source/Editor/FoliageEdit/Private/FoliageEditModule.cpp
+++ b/Engine/Source/Editor/FoliageEdit/Private/FoliageEditModule.cpp
@@ -19,6 +19,8 @@ const FName FoliageEditAppIdentifier = FName(TEXT("FoliageEdApp"));
 #include "ProceduralFoliageBlockingVolume.h"
 #include "ProceduralFoliageComponent.h"
 #include "FoliageTypeObjectCustomization.h"
+#include "FoliageType_InstancedStaticMesh.h"
+#include "FoliageType_ISMThumbnailRenderer.h"
 
 /**
  * Foliage Edit Mode module
@@ -63,6 +65,9 @@ public:
 
 		SubscribeEvents();
 #endif
+
+		// Register thumbnail renderer
+		UThumbnailManager::Get().RegisterCustomRenderer(UFoliageType_InstancedStaticMesh::StaticClass(), UFoliageType_ISMThumbnailRenderer::StaticClass());
 	}
 
 	/**
diff --git a/Engine/Source/Editor/FoliageEdit/Private/FoliagePaletteItem.cpp b/Engine/Source/Editor/FoliageEdit/Private/FoliagePaletteItem.cpp
--- a/Engine/Source/Editor/FoliageEdit/Private/FoliagePaletteItem.cpp
+++ b/Engine/Source/Editor/FoliageEdit/Private/FoliagePaletteItem.cpp
@@ -48,7 +48,7 @@ FFoliagePaletteItemModel::FFoliagePaletteItemModel(FFoliageMeshUIInfoPtr InTypeI
 	FAssetData AssetData;
 	if (!IsBlueprint())
 	{
-		AssetData = FAssetData(TypeInfo->Settings->GetStaticMesh());
+		AssetData = FAssetData(TypeInfo->Settings);
 	}
 	else
 	{
diff --git a/Engine/Source/Editor/FoliageEdit/Private/FoliageType_ISMThumbnailRenderer.cpp b/Engine/Source/Editor/FoliageEdit/Private/FoliageType_ISMThumbnailRenderer.cpp
new file mode 100644
--- /dev/null
+++ b/Engine/Source/Editor/FoliageEdit/Private/FoliageType_ISMThumbnailRenderer.cpp
@@ -0,0 +1,56 @@
+// Copyright 1998-2015 Epic Games, Inc. All Rights Reserved.
+
+#include "UnrealEd.h"
+
+// FPreviewScene derived helpers for rendering
+#include "ThumbnailHelpers.h"
+#include "EngineModule.h"
+#include "RendererInterface.h"
+#include "Engine/StaticMesh.h"
+#include "FoliageType_ISMThumbnailRenderer.h"
+#include "FoliageType_InstancedStaticMesh.h"
+
+bool UFoliageType_ISMThumbnailRenderer::CanVisualizeAsset(UObject* Object)
+{
+	UFoliageType_InstancedStaticMesh* FoliageType = Cast<UFoliageType_InstancedStaticMesh>(Object);
+	return FoliageType && FoliageType->GetStaticMesh();
+}
+
+void UFoliageType_ISMThumbnailRenderer::Draw(UObject* Object, int32 X, int32 Y, uint32 Width, uint32 Height, FRenderTarget* RenderTarget, FCanvas* Canvas)
+{
+	UFoliageType_InstancedStaticMesh* FoliageType = Cast<UFoliageType_InstancedStaticMesh>(Object);
+	if (FoliageType && FoliageType->GetStaticMesh())
+	{
+		if (ThumbnailScene == nullptr)
+		{
+			ThumbnailScene = new FStaticMeshThumbnailScene();
+		}
+
+		ThumbnailScene->SetStaticMesh(FoliageType->GetStaticMesh());
+		ThumbnailScene->SetOverrideMaterials(FoliageType->OverrideMaterials);
+		ThumbnailScene->GetScene()->UpdateSpeedTreeWind(0.0);
+
+		FSceneViewFamilyContext ViewFamily(FSceneViewFamily::ConstructionValues(RenderTarget, ThumbnailScene->GetScene(), FEngineShowFlags(ESFIM_Game))
+			.SetWorldTimes(FApp::GetCurrentTime() - GStartTime, FApp::GetDeltaTime(), FApp::GetCurrentTime() - GStartTime));
+
+		ViewFamily.EngineShowFlags.DisableAdvancedFeatures();
+		ViewFamily.EngineShowFlags.MotionBlur = 0;
+		ViewFamily.EngineShowFlags.LOD = 0;
+
+		ThumbnailScene->GetView(&ViewFamily, X, Y, Width, Height);
+		GetRendererModule().BeginRenderingViewFamily(Canvas, &ViewFamily);
+		ThumbnailScene->SetStaticMesh(nullptr);
+		ThumbnailScene->SetOverrideMaterials(TArray<class UMaterialInterface*>());
+	}
+}
+
+void UFoliageType_ISMThumbnailRenderer::BeginDestroy()
+{
+	if (ThumbnailScene != nullptr)
+	{
+		delete ThumbnailScene;
+		ThumbnailScene = nullptr;
+	}
+
+	Super::BeginDestroy();
+}
diff --git a/Engine/Source/Editor/FoliageEdit/Public/FoliageType_ISMThumbnailRenderer.h b/Engine/Source/Editor/FoliageEdit/Public/FoliageType_ISMThumbnailRenderer.h
new file mode 100644
--- /dev/null
+++ b/Engine/Source/Editor/FoliageEdit/Public/FoliageType_ISMThumbnailRenderer.h
@@ -0,0 +1,28 @@
+// Copyright 1998-2015 Epic Games, Inc. All Rights Reserved.
+
+/**
+ * This thumbnail renderer displays the static mesh used by this foliage type
+ */
+
+#pragma once
+#include "FoliageType_ISMThumbnailRenderer.generated.h"
+
+UCLASS(CustomConstructor, Config=Editor)
+class UFoliageType_ISMThumbnailRenderer : public UDefaultSizedThumbnailRenderer
+{
+	GENERATED_UCLASS_BODY()
+	
+	UFoliageType_ISMThumbnailRenderer(const FObjectInitializer& ObjectInitializer)
+	:	Super(ObjectInitializer)
+	,	ThumbnailScene(nullptr)
+	{}
+
+	// UThumbnailRenderer implementation
+	virtual void Draw(UObject* Object, int32 X, int32 Y, uint32 Width, uint32 Height, FRenderTarget*, FCanvas* Canvas) override;
+	virtual bool CanVisualizeAsset(UObject* Object) override;
+	// UObject implementation
+	virtual void BeginDestroy() override;
+
+private:
+	class FStaticMeshThumbnailScene* ThumbnailScene;
+};
diff --git a/Engine/Source/Editor/UnrealEd/Private/ThumbnailHelpers.cpp b/Engine/Source/Editor/UnrealEd/Private/ThumbnailHelpers.cpp
--- a/Engine/Source/Editor/UnrealEd/Private/ThumbnailHelpers.cpp
+++ b/Engine/Source/Editor/UnrealEd/Private/ThumbnailHelpers.cpp
@@ -550,6 +550,12 @@ void FStaticMeshThumbnailScene::SetStaticMesh(UStaticMesh* StaticMesh)
 	}
 }
 
+void FStaticMeshThumbnailScene::SetOverrideMaterials(const TArray<class UMaterialInterface*>& OverrideMaterials)
+{
+	PreviewActor->GetStaticMeshComponent()->OverrideMaterials = OverrideMaterials;
+	PreviewActor->GetStaticMeshComponent()->MarkRenderStateDirty();
+}
+
 void FStaticMeshThumbnailScene::GetViewMatrixParameters(const float InFOVDegrees, FVector& OutOrigin, float& OutOrbitPitch, float& OutOrbitYaw, float& OutOrbitZoom) const
 {
 	check(PreviewActor);
diff --git a/Engine/Source/Editor/UnrealEd/Public/ThumbnailHelpers.h b/Engine/Source/Editor/UnrealEd/Public/ThumbnailHelpers.h
--- a/Engine/Source/Editor/UnrealEd/Public/ThumbnailHelpers.h
+++ b/Engine/Source/Editor/UnrealEd/Public/ThumbnailHelpers.h
@@ -119,6 +119,9 @@ public:
 	/** Sets the static mesh to use in the next GetView() */
 	void SetStaticMesh(class UStaticMesh* StaticMesh);
 
+	/** Sets override materials for the static mesh  */
+	void SetOverrideMaterials(const TArray<class UMaterialInterface*>& OverrideMaterials);
+
 protected:
 	// FThumbnailPreviewScene implementation
 	virtual void GetViewMatrixParameters(const float InFOVDegrees, FVector& OutOrigin, float& OutOrbitPitch, float& OutOrbitYaw, float& OutOrbitZoom) const override;
diff --git a/Engine/Source/Runtime/Engine/Classes/Components/HierarchicalInstancedStaticMeshComponent.h b/Engine/Source/Runtime/Engine/Classes/Components/HierarchicalInstancedStaticMeshComponent.h
--- a/Engine/Source/Runtime/Engine/Classes/Components/HierarchicalInstancedStaticMeshComponent.h
+++ b/Engine/Source/Runtime/Engine/Classes/Components/HierarchicalInstancedStaticMeshComponent.h
@@ -130,7 +130,7 @@ public:
 	virtual bool RemoveInstance(int32 InstanceIndex) override;
 	virtual bool UpdateInstanceTransform(int32 InstanceIndex, const FTransform& NewInstanceTransform, bool bWorldSpace, bool bMarkRenderStateDirty = false) override;
 	virtual void ClearInstances() override;
-
+	virtual TArray<int32> GetInstancesOverlappingSphere(const FVector& Center, float Radius, bool bSphereInWorldSpace = true) const override;
 
 	/** Get the number of instances that overlap a given sphere */
 	int32 GetOverlappingSphereCount(const FSphere& Sphere) const;
diff --git a/Engine/Source/Runtime/Engine/Classes/Components/InstancedStaticMeshComponent.h b/Engine/Source/Runtime/Engine/Classes/Components/InstancedStaticMeshComponent.h
--- a/Engine/Source/Runtime/Engine/Classes/Components/InstancedStaticMeshComponent.h
+++ b/Engine/Source/Runtime/Engine/Classes/Components/InstancedStaticMeshComponent.h
@@ -60,7 +60,7 @@ struct FInstancedStaticMeshMappingInfo
 };
 
 /** A component that efficiently renders multiple instances of the same StaticMesh. */
-UCLASS(ClassGroup=Rendering, meta=(BlueprintSpawnableComponent))
+UCLASS(ClassGroup = Rendering, meta = (BlueprintSpawnableComponent), Blueprintable)
 class ENGINE_API UInstancedStaticMeshComponent : public UStaticMeshComponent
 {
 	GENERATED_UCLASS_BODY()
@@ -105,8 +105,16 @@ class ENGINE_API UInstancedStaticMeshComponent : public UStaticMeshComponent
 	/** Get the transform for the instance specified. Instance is returned in local space of this component unless bWorldSpace is set.  Returns True on success. */
 	UFUNCTION(BlueprintCallable, Category = "Components|InstancedStaticMesh")
 	bool GetInstanceTransform(int32 InstanceIndex, FTransform& OutInstanceTransform, bool bWorldSpace = false) const;
-	
-	/** Update the transform for the instance specified. Instance is given in local space of this component unless bWorldSpace is set.  Returns True on success. */
+
+	/**
+	* Update the transform for the instance specified.
+	*
+	* @param InstanceIndex			The index of the instance to update
+	* @param NewInstanceTransform	The new transform
+	* @param bWorldSpace			If true, the new transform interpreted as a World Space transform, otherwise it is interpreted as Local Space
+	* @param bMarkRenderStateDirty	If true, the change should be visible immediately. If you are updating many instances you should only set this to true for the last instance.
+	* @return						True on success.
+	*/
 	UFUNCTION(BlueprintCallable, Category = "Components|InstancedStaticMesh")
 	virtual bool UpdateInstanceTransform(int32 InstanceIndex, const FTransform& NewInstanceTransform, bool bWorldSpace=false, bool bMarkRenderStateDirty = false);
 
@@ -126,6 +134,10 @@ class ENGINE_API UInstancedStaticMeshComponent : public UStaticMeshComponent
 	UFUNCTION(BlueprintCallable, Category = "Components|InstancedStaticMesh")
 	void SetCullDistances(int32 StartCullDistance, int32 EndCullDistance);
 
+	/** Returns the instances with instance bounds overlapping the specified sphere. The return value is an array of instance indices. */
+	UFUNCTION(BlueprintCallable, Category = "Components|InstancedStaticMesh")
+	virtual TArray<int32> GetInstancesOverlappingSphere(const FVector& Center, float Radius, bool bSphereInWorldSpace=true) const;
+
 	virtual bool ShouldCreatePhysicsState() const override;
 
 public:
diff --git a/Engine/Source/Runtime/Engine/Private/HierarchicalInstancedStaticMesh.cpp b/Engine/Source/Runtime/Engine/Private/HierarchicalInstancedStaticMesh.cpp
--- a/Engine/Source/Runtime/Engine/Private/HierarchicalInstancedStaticMesh.cpp
+++ b/Engine/Source/Runtime/Engine/Private/HierarchicalInstancedStaticMesh.cpp
@@ -2242,51 +2242,51 @@ static void GatherInstanceTransformsInArea(const UHierarchicalInstancedStaticMes
 	const TArray<FClusterNode>& ClusterTree = *Component.ClusterTreePtr;
 	if (ClusterTree.Num())
 	{
-	const FClusterNode& ChildNode = ClusterTree[Child];
-	const FBox WorldNodeBox = FBox(ChildNode.BoundMin, ChildNode.BoundMax).TransformBy(Component.ComponentToWorld);
+		const FClusterNode& ChildNode = ClusterTree[Child];
+		const FBox WorldNodeBox = FBox(ChildNode.BoundMin, ChildNode.BoundMax).TransformBy(Component.ComponentToWorld);
 	
-	if (AreaBox.Intersect(WorldNodeBox))
-	{
-		if (ChildNode.FirstChild < 0 || AreaBox.IsInside(WorldNodeBox))
+		if (AreaBox.Intersect(WorldNodeBox))
 		{
-			// Unfortunately ordering of PerInstanceSMData does not match ordering of cluster tree, so we have to use remaping
-			const bool bUseRemaping = Component.SortedInstances.Num() > 0;
-			
-			// In case there no more subdivision or node is completely encapsulated by a area box
-			// add all instances to the result
-			for (int32 i = ChildNode.FirstInstance; i <= ChildNode.LastInstance; ++i)
+			if (ChildNode.FirstChild < 0 || AreaBox.IsInside(WorldNodeBox))
 			{
-				int32 SortedIdx = bUseRemaping ? Component.SortedInstances[i] : i;
+				// Unfortunately ordering of PerInstanceSMData does not match ordering of cluster tree, so we have to use remaping
+				const bool bUseRemaping = Component.SortedInstances.Num() > 0;
+			
+				// In case there no more subdivision or node is completely encapsulated by a area box
+				// add all instances to the result
+				for (int32 i = ChildNode.FirstInstance; i <= ChildNode.LastInstance; ++i)
+				{
+					int32 SortedIdx = bUseRemaping ? Component.SortedInstances[i] : i;
 
-					FTransform InstanceToComponent;
-					if (Component.PerInstanceSMData.IsValidIndex(SortedIdx))
-					{
-						InstanceToComponent = FTransform(Component.PerInstanceSMData[SortedIdx].Transform);
-					}
-					else if (Component.PerInstanceRenderData.IsValid())
+						FTransform InstanceToComponent;
+						if (Component.PerInstanceSMData.IsValidIndex(SortedIdx))
+						{
+							InstanceToComponent = FTransform(Component.PerInstanceSMData[SortedIdx].Transform);
+						}
+						else if (Component.PerInstanceRenderData.IsValid())
+						{
+							// if there's no PerInstanceSMData (e.g. for grass), we'll go ge the transform from the render buffer
+							FInstanceStream const* Inst = Component.PerInstanceRenderData->InstanceBuffer.GetInstance(i);
+							FMatrix XformMat;
+							Inst->GetInstanceTransform(XformMat);
+							InstanceToComponent = FTransform(XformMat);
+						}
+					if (!InstanceToComponent.GetScale3D().IsZero())
 					{
-						// if there's no PerInstanceSMData (e.g. for grass), we'll go ge the transform from the render buffer
-						FInstanceStream const* Inst = Component.PerInstanceRenderData->InstanceBuffer.GetInstance(i);
-						FMatrix XformMat;
-						Inst->GetInstanceTransform(XformMat);
-						InstanceToComponent = FTransform(XformMat);
+						InstanceData.Add(InstanceToComponent*Component.ComponentToWorld);
 					}
-				if (!InstanceToComponent.GetScale3D().IsZero())
-				{
-					InstanceData.Add(InstanceToComponent*Component.ComponentToWorld);
 				}
 			}
-		}
-		else
-		{
-			for (int32 i = ChildNode.FirstChild; i <= ChildNode.LastChild; ++i)
+			else
 			{
-				GatherInstanceTransformsInArea(Component, AreaBox, i, InstanceData);
+				for (int32 i = ChildNode.FirstChild; i <= ChildNode.LastChild; ++i)
+				{
+					GatherInstanceTransformsInArea(Component, AreaBox, i, InstanceData);
+				}
 			}
 		}
 	}
 }
-}
 
 int32 UHierarchicalInstancedStaticMeshComponent::GetOverlappingSphereCount(const FSphere& Sphere) const
 {
@@ -2409,6 +2409,66 @@ void UHierarchicalInstancedStaticMeshComponent::FlushAccumulatedNavigationUpdate
 	}
 }
 
+// recursive helper to gather all instances with locations inside the specified sphere
+static void GatherInstancesOverlappingSphere(const UHierarchicalInstancedStaticMeshComponent& Component, const FSphere& Sphere, const float StaticMeshBoundsRadius, const FBox& AreaBox, int32 Child, TArray<int32>& OutInstanceIndices)
+{
+	const TArray<FClusterNode>& ClusterTree = *Component.ClusterTreePtr;
+	const FClusterNode& ChildNode = ClusterTree[Child];
+	const FBox WorldNodeBox = FBox(ChildNode.BoundMin, ChildNode.BoundMax).TransformBy(Component.ComponentToWorld);
+
+	if (AreaBox.Intersect(WorldNodeBox))
+	{
+		if (ChildNode.FirstChild < 0 || AreaBox.IsInside(WorldNodeBox))
+		{
+			// Unfortunately ordering of PerInstanceSMData does not match ordering of cluster tree, so we have to use remaping
+			const bool bUseRemaping = Component.SortedInstances.Num() > 0;
+
+			// In case there no more subdivision or node is completely encapsulated by a area box
+			// add all instances to the result
+			for (int32 i = ChildNode.FirstInstance; i <= ChildNode.LastInstance; ++i)
+			{
+				int32 SortedIdx = bUseRemaping ? Component.SortedInstances[i] : i;
+				if (Component.PerInstanceSMData.IsValidIndex(SortedIdx))
+				{
+					const FMatrix& Matrix = Component.PerInstanceSMData[SortedIdx].Transform;
+					FSphere InstanceSphere(Matrix.GetOrigin(), StaticMeshBoundsRadius * Matrix.GetScaleVector().GetMax());
+					if (Sphere.Intersects(InstanceSphere))
+					{
+						OutInstanceIndices.Add(SortedIdx);
+					}
+				}
+			}
+		}
+		else
+		{
+			for (int32 i = ChildNode.FirstChild; i <= ChildNode.LastChild; ++i)
+			{
+				GatherInstancesOverlappingSphere(Component, Sphere, StaticMeshBoundsRadius, AreaBox, i, OutInstanceIndices);
+			}
+		}
+	}
+}
+
+TArray<int32> UHierarchicalInstancedStaticMeshComponent::GetInstancesOverlappingSphere(const FVector& Center, float Radius, bool bSphereInWorldSpace) const
+{
+	if (ClusterTreePtr.IsValid() && ClusterTreePtr->Num())
+	{
+		TArray<int32> Result;
+		FSphere Sphere(Center, Radius);
+		if (bSphereInWorldSpace)
+		{
+			Sphere = Sphere.TransformBy(ComponentToWorld.Inverse());
+		}
+		const FBox AABB(Sphere.Center - FVector(Sphere.W), Sphere.Center + FVector(Sphere.W));
+		GatherInstancesOverlappingSphere(*this, Sphere, StaticMesh->GetBounds().SphereRadius, AABB, 0, Result);
+		return Result;
+	}
+	else
+	{
+		return Super::GetInstancesOverlappingSphere(Center, Radius, bSphereInWorldSpace);
+	}
+}
+
 static void RebuildFoliageTrees(const TArray<FString>& Args)
 {
 #if !(UE_BUILD_SHIPPING || UE_BUILD_TEST)
diff --git a/Engine/Source/Runtime/Engine/Private/InstancedStaticMesh.cpp b/Engine/Source/Runtime/Engine/Private/InstancedStaticMesh.cpp
--- a/Engine/Source/Runtime/Engine/Private/InstancedStaticMesh.cpp
+++ b/Engine/Source/Runtime/Engine/Private/InstancedStaticMesh.cpp
@@ -1481,6 +1481,32 @@ bool UInstancedStaticMeshComponent::UpdateInstanceTransform(int32 InstanceIndex,
 	return true;
 }
 
+TArray<int32> UInstancedStaticMeshComponent::GetInstancesOverlappingSphere(const FVector& Center, float Radius, bool bSphereInWorldSpace) const
+{
+	TArray<int32> Result;
+
+	FSphere Sphere(Center, Radius);
+	if (bSphereInWorldSpace)
+	{
+		Sphere = Sphere.TransformBy(ComponentToWorld.Inverse());
+	}
+
+	float StaticMeshBoundsRadius = StaticMesh->GetBounds().SphereRadius;
+
+	for (int32 Index = 0; Index < PerInstanceSMData.Num(); Index++)
+	{
+		const FMatrix& Matrix = PerInstanceSMData[Index].Transform;
+		FSphere InstanceSphere(Matrix.GetOrigin(), StaticMeshBoundsRadius * Matrix.GetScaleVector().GetMax());
+
+		if (Sphere.Intersects(InstanceSphere))
+		{
+			Result.Add(Index);
+		}
+	}
+
+	return Result;
+}
+
 bool UInstancedStaticMeshComponent::ShouldCreatePhysicsState() const
 {
 	return IsRegistered() && !IsBeingDestroyed() && (bAlwaysCreatePhysicsState || IsCollisionEnabled());
diff --git a/Engine/Source/Runtime/Foliage/Private/InstancedFoliage.cpp b/Engine/Source/Runtime/Foliage/Private/InstancedFoliage.cpp
--- a/Engine/Source/Runtime/Foliage/Private/InstancedFoliage.cpp
+++ b/Engine/Source/Runtime/Foliage/Private/InstancedFoliage.cpp
@@ -10,7 +10,7 @@
 #include "UObjectToken.h"
 #include "MapErrors.h"
 #include "Components/ModelComponent.h"
-#include "Components/HierarchicalInstancedStaticMeshComponent.h"
+#include "FoliageInstancedStaticMeshComponent.h"
 #include "Serialization/CustomVersion.h"
 #include "ProceduralFoliageComponent.h"
 #include "ProceduralFoliageBlockingVolume.h"
@@ -48,6 +48,8 @@ struct FFoliageCustomVersion
 		FoliageTypeCustomizationScaling = 7,
 		// FoliageType procedural scale and shade settings updated
 		FoliageTypeProceduralScaleAndShade = 8,
+		// Added FoliageHISMC and blueprint support
+		FoliageHISMCBlueprints = 9,
 		// -----<new versions can be added above this line>-------------------------------------------------
 		VersionPlusOne,
 		LatestVersion = VersionPlusOne - 1
@@ -453,6 +455,7 @@ UFoliageType_InstancedStaticMesh::UFoliageType_InstancedStaticMesh(const FObject
 	: Super(ObjectInitializer)
 {
 	Mesh = nullptr;
+	ComponentClass = UFoliageInstancedStaticMeshComponent::StaticClass();
 }
 
 
@@ -609,6 +612,75 @@ void FFoliageMeshInfo::CheckValid()
 #endif
 }
 
+void FFoliageMeshInfo::CreateNewComponent(AInstancedFoliageActor* InIFA, const UFoliageType* InSettings)
+{
+	check(!Component);
+
+	UClass* ComponentClass = InSettings->GetComponentClass();
+	if (ComponentClass == nullptr)
+	{
+		ComponentClass = UFoliageInstancedStaticMeshComponent::StaticClass();
+	}
+
+	Component = NewObject<UFoliageInstancedStaticMeshComponent>(InIFA, ComponentClass, NAME_None, RF_Transactional);
+
+	Component->StaticMesh = InSettings->GetStaticMesh();
+	Component->bSelectable = true;
+	Component->bHasPerInstanceHitProxies = true;
+	Component->InstancingRandomSeed = FMath::Rand();
+
+	UpdateComponentSettings(InSettings);
+
+	Component->AttachTo(InIFA->GetRootComponent());
+
+	if (InIFA->GetRootComponent()->IsRegistered())
+	{
+		Component->RegisterComponent();
+	}
+
+	// Use only instance translation as a component transform
+	Component->SetWorldTransform(InIFA->GetRootComponent()->ComponentToWorld);
+
+	// Add the new component to the transaction buffer so it will get destroyed on undo
+	Component->Modify();
+	// We don't want to track changes to instances later so we mark it as non-transactional
+	Component->ClearFlags(RF_Transactional);
+}
+
+void FFoliageMeshInfo::CheckComponentClass(AInstancedFoliageActor* InIFA, const UFoliageType* InSettings)
+{
+	if (Component)
+	{
+		UClass* ComponentClass = InSettings->GetComponentClass();
+		if (ComponentClass == nullptr)
+		{
+			ComponentClass = UFoliageInstancedStaticMeshComponent::StaticClass();
+		}
+
+		if(ComponentClass != Component->GetClass())
+		{
+			InIFA->Modify();
+
+			// prepare to destroy the old component
+			Component->ClearInstances();
+
+			// make sure the destruction gets stored in the undo buffer, 
+			// so the component will get recreated on undo.
+			Component->SetFlags(RF_Transactional);
+			Component->Modify();
+			
+			Component->DestroyComponent();
+			Component = nullptr;
+
+			// create a new component
+			CreateNewComponent(InIFA, InSettings);
+
+			// apply the instances to it
+			ReapplyInstancesToComponent();
+		}
+	}
+}
+
 void FFoliageMeshInfo::UpdateComponentSettings(const UFoliageType* InSettings)
 {
 	if (Component)
@@ -702,6 +774,31 @@ void FFoliageMeshInfo::UpdateComponentSettings(const UFoliageType* InSettings)
 			bNeedsMarkRenderStateDirty = true;
 			bNeedsInvalidateLightingCache = true;
 		}
+
+		const UFoliageType_InstancedStaticMesh* FoliageType_ISM = Cast<UFoliageType_InstancedStaticMesh>(FoliageType);
+		if (FoliageType_ISM)
+		{
+			// Check override materials
+			if (Component->OverrideMaterials.Num() != FoliageType_ISM->OverrideMaterials.Num())
+			{
+				Component->OverrideMaterials = FoliageType_ISM->OverrideMaterials;
+				bNeedsMarkRenderStateDirty = true;
+				bNeedsInvalidateLightingCache = true;
+			}
+			else
+			{
+				for (int32 Index = 0; Index < FoliageType_ISM->OverrideMaterials.Num(); Index++)
+				{
+					if (Component->OverrideMaterials[Index] != FoliageType_ISM->OverrideMaterials[Index])
+					{
+						Component->OverrideMaterials = FoliageType_ISM->OverrideMaterials;
+						bNeedsMarkRenderStateDirty = true;
+						bNeedsInvalidateLightingCache = true;
+						break;
+					}
+				}
+			}
+		}	
 		
 		Component->BodyInstance.CopyBodyInstancePropertiesFrom(&FoliageType->BodyInstance);
 
@@ -732,29 +829,8 @@ void FFoliageMeshInfo::AddInstance(AInstancedFoliageActor* InIFA, const UFoliage
 
 	if (Component == nullptr)
 	{
-		Component = NewObject<UHierarchicalInstancedStaticMeshComponent>(InIFA, NAME_None, RF_Transactional);
-
-		Component->StaticMesh = InSettings->GetStaticMesh();
-		Component->bSelectable = true;
-		Component->bHasPerInstanceHitProxies = true;
-		Component->InstancingRandomSeed = FMath::Rand();
-
-		UpdateComponentSettings(InSettings);
-
-		Component->AttachTo(InIFA->GetRootComponent());
-		
-		if (InIFA->GetRootComponent()->IsRegistered())
-		{
-			Component->RegisterComponent();
-		}
-
-		// Use only instance translation as a component transform
-		Component->SetWorldTransform(InIFA->GetRootComponent()->ComponentToWorld);
-
-		// Add the new component to the transaction buffer so it will get destroyed on undo
-		Component->Modify();
-		// We don't want to track changes to instances later so we mark it as non-transactional
-		Component->ClearFlags(RF_Transactional);
+		CreateNewComponent(InIFA, InSettings);
+		check(Component);
 	}
 	else
 	{
@@ -774,12 +850,6 @@ void FFoliageMeshInfo::AddInstance(AInstancedFoliageActor* InIFA, const UFoliage
 	// Add the instance to the component
 	Component->AddInstanceWorldSpace(InstanceToWorld);
 
-	// Update PrimitiveComponent's culling distance taking into account the radius of the bounds, as
-	// it is based on the center of the component's bounds.
-//	float CullDistance = InSettings->EndCullDistance > 0 ? (float)InSettings->EndCullDistance + BestCluster->Bounds.SphereRadius : 0.f;
-//	Component->LDMaxDrawDistance = CullDistance;
-//	Component->CachedMaxDrawDistance = CullDistance;
-
 	CheckValid();
 }
 
@@ -999,29 +1069,6 @@ void FFoliageMeshInfo::ReapplyInstancesToComponent()
 	}
 }
 
-
-// Update settings in the clusters based on the current settings (eg culling distance, collision, ...)
-//void FFoliageMeshInfo::UpdateClusterSettings(AInstancedFoliageActor* InIFA)
-//{
-//	for (FFoliageInstanceCluster& Cluster : InstanceClusters)
-//	{
-//		UInstancedStaticMeshComponent* ClusterComponent = Cluster.ClusterComponent;
-//		ClusterComponent->MarkRenderStateDirty();
-//
-//		// Copy settings
-//		ClusterComponent->InstanceStartCullDistance = Settings->StartCullDistance;
-//		ClusterComponent->InstanceEndCullDistance = Settings->EndCullDistance;
-//
-//		// Update PrimitiveComponent's culling distance taking into account the radius of the bounds, as
-//		// it is based on the center of the component's bounds.
-//		float CullDistance = Settings->EndCullDistance > 0 ? (float)Settings->EndCullDistance + Cluster.Bounds.SphereRadius : 0.f;
-//		ClusterComponent->LDMaxDrawDistance = CullDistance;
-//		ClusterComponent->CachedMaxDrawDistance = CullDistance;
-//	}
-//
-//	InIFA->MarkComponentsRenderStateDirty();
-//}
-
 void FFoliageMeshInfo::GetInstancesInsideSphere(const FSphere& Sphere, TArray<int32>& OutInstances)
 {
 	auto TempInstances = InstanceHash->GetInstancesOverlappingBox(FBox::BuildAABB(Sphere.Center, FVector(Sphere.W)));
@@ -1667,28 +1714,6 @@ void UpdateSettingsBounds(const UStaticMesh* InMesh, UFoliageType_InstancedStati
 	Settings->LowBoundOriginRadius = FVector((MinX + MaxX), (MinY + MaxY), FMath::Sqrt(FMath::Square(MaxX - MinX) + FMath::Square(MaxY - MinY))) * 0.5f;
 }
 
-
-#if WITH_EDITORONLY_DATA
-FFoliageMeshInfo* AInstancedFoliageActor::UpdateMeshSettings(const UStaticMesh* InMesh, const UFoliageType_InstancedStaticMesh* DefaultSettings, bool bIncludeAssets)
-{
-	if (UFoliageType* OldSettings = GetLocalFoliageTypeForMesh(InMesh, nullptr))
-	{
-		MarkPackageDirty();
-
-		UFoliageType_InstancedStaticMesh* NewSettings = DuplicateObject<UFoliageType_InstancedStaticMesh>(DefaultSettings, this);
-		UpdateSettingsBounds(InMesh, NewSettings);
-
-		TUniqueObj<FFoliageMeshInfo> MeshInfo;
-		FoliageMeshes.RemoveAndCopyValue(OldSettings, MeshInfo);
-		MeshInfo->FoliageTypeUpdateGuid = NewSettings->UpdateGuid;
-		MeshInfo->UpdateComponentSettings(NewSettings);
-		return &*FoliageMeshes.Add(NewSettings, MoveTemp(MeshInfo));
-	}
-
-	return nullptr;
-}
-#endif
-
 UFoliageType* AInstancedFoliageActor::AddFoliageType(const UFoliageType* InType, FFoliageMeshInfo** OutInfo)
 {
 	FFoliageMeshInfo* MeshInfo = nullptr;
@@ -1915,6 +1940,9 @@ void AInstancedFoliageActor::PostEditUndo()
 	for (auto& MeshPair : FoliageMeshes)
 	{
 		FFoliageMeshInfo& MeshInfo = *MeshPair.Value;
+		
+		MeshInfo.CheckComponentClass(this, MeshPair.Key);
+
 		MeshInfo.ReapplyInstancesToComponent();
 
 		// Regenerate instance hash
@@ -2267,6 +2295,7 @@ void AInstancedFoliageActor::PostLoad()
 			{
 				if (MeshInfo.FoliageTypeUpdateGuid.IsValid())
 				{
+					MeshInfo.CheckComponentClass(this, FoliageType);
 					MeshInfo.UpdateComponentSettings(FoliageType);
 				}
 				MeshInfo.FoliageTypeUpdateGuid = FoliageType->UpdateGuid;
@@ -2318,6 +2347,7 @@ void AInstancedFoliageActor::NotifyFoliageTypeChanged(UFoliageType* FoliageType,
 
 	if (TypeInfo)
 	{
+		TypeInfo->CheckComponentClass(this, FoliageType);
 		TypeInfo->UpdateComponentSettings(FoliageType);
 
 		if (bMeshChanged)
@@ -2633,4 +2663,60 @@ bool FPotentialInstance::PlaceInstance(const UWorld* InWorld, const UFoliageType
 	return bSkipCollision || AInstancedFoliageActor::CheckCollisionWithWorld(InWorld, Settings, Inst, HitNormal, HitLocation);
 }
 
+float AInstancedFoliageActor::InternalTakeRadialDamage(float Damage, struct FRadialDamageEvent const& RadialDamageEvent, class AController* EventInstigator, AActor* DamageCauser)
+{
+	// Radial damage scaling needs to be applied per instance so we don't do anything here
+	return Damage;
+}
+
+UFoliageInstancedStaticMeshComponent::UFoliageInstancedStaticMeshComponent(const FObjectInitializer& ObjectInitializer)
+: Super(ObjectInitializer)
+{
+}
+
+void UFoliageInstancedStaticMeshComponent::ReceiveComponentDamage(float DamageAmount, FDamageEvent const& DamageEvent, AController* EventInstigator, AActor* DamageCauser)
+{
+	Super::ReceiveComponentDamage(DamageAmount, DamageEvent, EventInstigator, DamageCauser);
+
+	if (DamageAmount != 0.f)
+	{
+		UDamageType const* const DamageTypeCDO = DamageEvent.DamageTypeClass ? DamageEvent.DamageTypeClass->GetDefaultObject<UDamageType>() : GetDefault<UDamageType>();
+		if (DamageEvent.IsOfType(FPointDamageEvent::ClassID))
+		{
+			// Point damage event, hit a single instance.
+			FPointDamageEvent* const PointDamageEvent = (FPointDamageEvent*)&DamageEvent;
+			if (PerInstanceSMData.IsValidIndex(PointDamageEvent->HitInfo.Item))
+			{
+				OnInstanceTakePointDamage.Broadcast(PointDamageEvent->HitInfo.Item, DamageAmount, EventInstigator, PointDamageEvent->HitInfo.ImpactPoint, PointDamageEvent->ShotDirection, DamageTypeCDO, DamageCauser);
+			}
+		}
+		else if (DamageEvent.IsOfType(FRadialDamageEvent::ClassID))
+		{
+			// Radial damage event, find which instances it hit and notify
+			FRadialDamageEvent* const RadialDamageEvent = (FRadialDamageEvent*)&DamageEvent;
+
+			float MaxRadius = RadialDamageEvent->Params.GetMaxRadius();
+			TArray<int32> Instances = GetInstancesOverlappingSphere(RadialDamageEvent->Origin, MaxRadius, true);
+
+			if (Instances.Num())
+			{
+				FVector LocalOrigin = GetComponentToWorld().Inverse().TransformPosition(RadialDamageEvent->Origin);
+				float Scale = GetComponentScale().X; // assume component (not instances) is uniformly scaled
+				
+				TArray<FFoliageInstanceDamage> DamageInstances;
+				DamageInstances.Empty(Instances.Num());
+				
+				for (int32 InstanceIndex : Instances)
+				{
+					// Find distance in local space and then scale; quicker than transforming each instance to world space.
+					float DistanceFromOrigin = (PerInstanceSMData[InstanceIndex].Transform.GetOrigin() - LocalOrigin).Size() * Scale;
+					new(DamageInstances) FFoliageInstanceDamage(InstanceIndex, RadialDamageEvent->Params.GetDamageScale(DistanceFromOrigin));
+				}
+
+				OnInstanceTakeRadialDamage.Broadcast(DamageInstances, EventInstigator, RadialDamageEvent->Origin, MaxRadius, DamageTypeCDO, DamageCauser);
+			}
+		}
+	}
+}
+
 #undef LOCTEXT_NAMESPACE
diff --git a/Engine/Source/Runtime/Foliage/Public/FoliageInstancedStaticMeshComponent.h b/Engine/Source/Runtime/Foliage/Public/FoliageInstancedStaticMeshComponent.h
new file mode 100644
--- /dev/null
+++ b/Engine/Source/Runtime/Foliage/Public/FoliageInstancedStaticMeshComponent.h
@@ -0,0 +1,49 @@
+// Copyright 1998-2015 Epic Games, Inc. All Rights Reserved.
+
+#pragma once
+
+#include "Components/HierarchicalInstancedStaticMeshComponent.h"
+#include "FoliageInstancedStaticMeshComponent.generated.h"
+
+USTRUCT(BlueprintType)
+struct FFoliageInstanceDamage
+{
+	GENERATED_USTRUCT_BODY()
+
+	FFoliageInstanceDamage()
+	: InstanceIndex(0)
+	, DamageAmount(0.f)
+	{}
+
+	FFoliageInstanceDamage(int32 InInstanceIndex, float InDamageAmount)
+	:	InstanceIndex(InInstanceIndex)
+	,	DamageAmount(InDamageAmount)
+	{}
+
+	UPROPERTY(EditAnywhere, BlueprintReadOnly, Category = FoliageInstanceDamage)
+	int32 InstanceIndex;
+
+	UPROPERTY(EditAnywhere, BlueprintReadOnly, Category = FoliageInstanceDamage)
+	float DamageAmount;
+};
+
+/**
+*/
+DECLARE_DYNAMIC_MULTICAST_DELEGATE_SevenParams(FInstancePointDamageSignature, int32, InstanceIndex, float, Damage, class AController*, InstigatedBy, FVector, HitLocation, FVector, ShotFromDirection, const class UDamageType*, DamageType, AActor*, DamageCauser);
+DECLARE_DYNAMIC_MULTICAST_DELEGATE_SixParams(FInstanceRadialDamageSignature, const TArray<FFoliageInstanceDamage>&, DamageInstances, class AController*, InstigatedBy, FVector, Origin, float, MaxRadius, const class UDamageType*, DamageType, AActor*, DamageCauser);
+
+UCLASS(ClassGroup = Foliage, Blueprintable)
+class FOLIAGE_API UFoliageInstancedStaticMeshComponent : public UHierarchicalInstancedStaticMeshComponent
+{
+	GENERATED_UCLASS_BODY()
+public:
+
+	UPROPERTY(BlueprintAssignable, Category = "Game|Damage")
+	FInstancePointDamageSignature OnInstanceTakePointDamage;
+
+	UPROPERTY(BlueprintAssignable, Category = "Game|Damage")
+	FInstanceRadialDamageSignature OnInstanceTakeRadialDamage;
+
+	virtual void ReceiveComponentDamage(float DamageAmount, FDamageEvent const& DamageEvent, AController* EventInstigator, AActor* DamageCauser) override;
+};
+
diff --git a/Engine/Source/Runtime/Foliage/Public/FoliageType.h b/Engine/Source/Runtime/Foliage/Public/FoliageType.h
--- a/Engine/Source/Runtime/Foliage/Public/FoliageType.h
+++ b/Engine/Source/Runtime/Foliage/Public/FoliageType.h
@@ -44,6 +44,9 @@ class UFoliageType : public UObject
 	virtual UStaticMesh* GetStaticMesh() const PURE_VIRTUAL(UFoliageType::GetStaticMesh, return nullptr; );
 	virtual void SetStaticMesh(UStaticMesh* InStaticMesh) PURE_VIRTUAL(UFoliageType::SetStaticMesh,);
 
+	/* Gets the component class to use for instances of this FoliageType */
+	virtual UClass* GetComponentClass() const PURE_VIRTUAL(UFoliageType::GetComponentClass(), return nullptr;);
+	
 	virtual void Serialize(FArchive& Ar) override;
 
 	virtual bool IsNotAssetOrBlueprint() const;
diff --git a/Engine/Source/Runtime/Foliage/Public/FoliageType_InstancedStaticMesh.h b/Engine/Source/Runtime/Foliage/Public/FoliageType_InstancedStaticMesh.h
--- a/Engine/Source/Runtime/Foliage/Public/FoliageType_InstancedStaticMesh.h
+++ b/Engine/Source/Runtime/Foliage/Public/FoliageType_InstancedStaticMesh.h
@@ -3,6 +3,7 @@
 #pragma once
 
 #include "FoliageType.h"
+#include "FoliageInstancedStaticMeshComponent.h"
 #include "FoliageType_InstancedStaticMesh.generated.h"
 
 UCLASS(hidecategories=Object, editinlinenew, MinimalAPI)
@@ -13,6 +14,14 @@ class UFoliageType_InstancedStaticMesh : public UFoliageType
 	UPROPERTY(EditAnywhere, Category=Mesh, meta=(DisplayThumbnail="true"))
 	UStaticMesh* Mesh;
 
+	UPROPERTY(EditAnywhere, AdvancedDisplay, Category = Mesh, Meta = (ToolTip = "Material overrides for foliage instances."))
+	TArray<class UMaterialInterface*> OverrideMaterials;
+		
+	/** The component class to use for foliage instances. 
+	  * You can make a Blueprint subclass of FoliageInstancedStaticMeshComponent to implement custom behavior and assign that class here. */
+	UPROPERTY(EditAnywhere, Category = Mesh)
+	TSubclassOf<UFoliageInstancedStaticMeshComponent> ComponentClass;
+
 	virtual UStaticMesh* GetStaticMesh() const override
 	{
 		return Mesh;
@@ -22,4 +31,9 @@ class UFoliageType_InstancedStaticMesh : public UFoliageType
 	{
 		Mesh = InStaticMesh;
 	}
+
+	virtual UClass* GetComponentClass() const override
+	{
+		return *ComponentClass;
+	}
 };
diff --git a/Engine/Source/Runtime/Foliage/Public/InstancedFoliage.h b/Engine/Source/Runtime/Foliage/Public/InstancedFoliage.h
--- a/Engine/Source/Runtime/Foliage/Public/InstancedFoliage.h
+++ b/Engine/Source/Runtime/Foliage/Public/InstancedFoliage.h
@@ -187,7 +187,10 @@ struct FFoliageMeshInfo
 	FOLIAGE_API void AddInstance(AInstancedFoliageActor* InIFA, const UFoliageType* InSettings, const FFoliageInstance& InNewInstance);
 	FOLIAGE_API void AddInstance(AInstancedFoliageActor* InIFA, const UFoliageType* InSettings, const FFoliageInstance& InNewInstance, UActorComponent* InBaseComponent);
 	FOLIAGE_API void RemoveInstances(AInstancedFoliageActor* InIFA, const TArray<int32>& InInstancesToRemove);
+	// Apply changes in the FoliageType to the component
 	FOLIAGE_API void UpdateComponentSettings(const UFoliageType* InSettings);
+	// Recreate the component if the FoliageType's ComponentClass doesn't match the Component's class
+	FOLIAGE_API void CheckComponentClass(AInstancedFoliageActor* InIFA, const UFoliageType* InSettings);
 	FOLIAGE_API void PreMoveInstances(AInstancedFoliageActor* InIFA, const TArray<int32>& InInstancesToMove);
 	FOLIAGE_API void PostMoveInstances(AInstancedFoliageActor* InIFA, const TArray<int32>& InInstancesMoved);
 	FOLIAGE_API void PostUpdateInstances(AInstancedFoliageActor* InIFA, const TArray<int32>& InInstancesUpdated, bool bReAddToHash = false);
@@ -210,7 +213,10 @@ struct FFoliageMeshInfo
 
 	FOLIAGE_API void AddToBaseHash(int32 InstanceIdx);
 	FOLIAGE_API void RemoveFromBaseHash(int32 InstanceIdx);
-		
+
+	// Create and register a new component
+	void CreateNewComponent(AInstancedFoliageActor* InIFA, const UFoliageType* InSettings);
+
 	// For debugging. Validate state after editing.
 	void CheckValid();
 #endif
diff --git a/Engine/Source/Runtime/Foliage/Public/InstancedFoliageActor.h b/Engine/Source/Runtime/Foliage/Public/InstancedFoliageActor.h
--- a/Engine/Source/Runtime/Foliage/Public/InstancedFoliageActor.h
+++ b/Engine/Source/Runtime/Foliage/Public/InstancedFoliageActor.h
@@ -49,6 +49,10 @@ public:
 	// we don't want to have our components automatically destroyed by the Blueprint code
 	virtual void RerunConstructionScripts() override {}
 	virtual bool IsLevelBoundsRelevant() const override { return false; }
+protected:
+	// Default InternalTakeRadialDamage behavior finds and scales damage for the closest component which isn't appropriate for foliage.
+	virtual float InternalTakeRadialDamage(float Damage, struct FRadialDamageEvent const& RadialDamageEvent, class AController* EventInstigator, AActor* DamageCauser) override;
+public:
 #if WITH_EDITOR
 	virtual void Destroyed() override;
 	FOLIAGE_API void CleanupDeletedFoliageType();
@@ -145,7 +149,6 @@ public:
 	// Add a new static mesh.
 	FOLIAGE_API FFoliageMeshInfo* AddMesh(UStaticMesh* InMesh, UFoliageType** OutSettings = nullptr, const UFoliageType_InstancedStaticMesh* DefaultSettings = nullptr);
 	FOLIAGE_API FFoliageMeshInfo* AddMesh(UFoliageType* InType);
-	FOLIAGE_API FFoliageMeshInfo* UpdateMeshSettings(const UStaticMesh* InMesh, const UFoliageType_InstancedStaticMesh* DefaultSettings, bool bIncludeAssets = true);
 
 	// Remove the FoliageType from the list, and all its instances.
 	FOLIAGE_API void RemoveFoliageType(UFoliageType** InFoliageType, int32 Num);
