Convert Destroyed behaviors in to EndPlay (GameplayAbilityTargetActor, MatineeActor, PlayerCameraManager, EmitterCameraLensEffectBase,
Make editor-only Destroyed not execute in PIE (Brush, CullDistanceVolume, InstancedFoliage, LandscapeProxy)

diff --git a/Engine/Source/Runtime/Engine/Classes/Camera/PlayerCameraManager.h b/Engine/Source/Runtime/Engine/Classes/Camera/PlayerCameraManager.h
--- a/Engine/Source/Runtime/Engine/Classes/Camera/PlayerCameraManager.h
+++ b/Engine/Source/Runtime/Engine/Classes/Camera/PlayerCameraManager.h
@@ -437,7 +437,7 @@ public:
 	// Begin AActor Interface
 	virtual bool ShouldTickIfViewportsOnly() const override;
 	virtual void PostInitializeComponents() override;
-	virtual void Destroyed() override;
+	virtual void EndPlay(const EEndPlayReason::Type EndPlayReason) override;
 	virtual void DisplayDebug(class UCanvas* Canvas, const FDebugDisplayInfo& DebugDisplay, float& YL, float& YPos) override;
 	virtual void ApplyWorldOffset(const FVector& InOffset, bool bWorldShift) override;
 	// End AActor Interface
diff --git a/Engine/Source/Runtime/Engine/Classes/Matinee/MatineeActor.h b/Engine/Source/Runtime/Engine/Classes/Matinee/MatineeActor.h
--- a/Engine/Source/Runtime/Engine/Classes/Matinee/MatineeActor.h
+++ b/Engine/Source/Runtime/Engine/Classes/Matinee/MatineeActor.h
@@ -268,7 +268,7 @@ public:
 	ENGINE_API void OnObjectsReplaced(const TMap<UObject*,UObject*>& ReplacementMap);
 #endif //WITH_EDITOR
 
-	virtual void Destroyed() override;
+	virtual void EndPlay(const EEndPlayReason::Type EndPlayReason) override;
 protected:
 
 	/** Handle for efficient management of CheckPriorityRefresh timer */
diff --git a/Engine/Source/Runtime/Engine/Classes/Particles/EmitterCameraLensEffectBase.h b/Engine/Source/Runtime/Engine/Classes/Particles/EmitterCameraLensEffectBase.h
--- a/Engine/Source/Runtime/Engine/Classes/Particles/EmitterCameraLensEffectBase.h
+++ b/Engine/Source/Runtime/Engine/Classes/Particles/EmitterCameraLensEffectBase.h
@@ -56,7 +56,7 @@ public:
 
 
 	// Begin AActor Interface
-	ENGINE_API virtual void Destroyed() override;
+	ENGINE_API virtual void EndPlay(const EEndPlayReason::Type EndPlayReason) override;
 	ENGINE_API virtual void PostInitializeComponents() override;
 	// End AActor Interface
 
diff --git a/Engine/Source/Runtime/Engine/Private/Brush.cpp b/Engine/Source/Runtime/Engine/Private/Brush.cpp
--- a/Engine/Source/Runtime/Engine/Private/Brush.cpp
+++ b/Engine/Source/Runtime/Engine/Private/Brush.cpp
@@ -195,7 +195,7 @@ void ABrush::Destroyed()
 {
 	Super::Destroyed();
 
-	if(IsStaticBrush())
+	if(GIsEditor && IsStaticBrush() && !GetWorld()->IsGameWorld())
 	{
 		// Trigger a csg rebuild if we're in the editor.
 		SetNeedRebuild(GetLevel());
diff --git a/Engine/Source/Runtime/Engine/Private/CullDistanceVolume.cpp b/Engine/Source/Runtime/Engine/Private/CullDistanceVolume.cpp
--- a/Engine/Source/Runtime/Engine/Private/CullDistanceVolume.cpp
+++ b/Engine/Source/Runtime/Engine/Private/CullDistanceVolume.cpp
@@ -21,9 +21,10 @@ void ACullDistanceVolume::Destroyed()
 {
 	Super::Destroyed();
 
-	if (GetWorld())
+	UWorld* World = GetWorld();
+	if (GIsEditor && World && World->IsGameWorld())
 	{
-		GetWorld()->bDoDelayedUpdateCullDistanceVolumes = true;
+		World->bDoDelayedUpdateCullDistanceVolumes = true;
 	}
 }
 
diff --git a/Engine/Source/Runtime/Engine/Private/Interpolation.cpp b/Engine/Source/Runtime/Engine/Private/Interpolation.cpp
--- a/Engine/Source/Runtime/Engine/Private/Interpolation.cpp
+++ b/Engine/Source/Runtime/Engine/Private/Interpolation.cpp
@@ -549,14 +549,15 @@ void AMatineeActor::OnObjectsReplaced(const TMap<UObject*,UObject*>& Replacement
 }
 #endif //WITH_EDITOR
 
-void AMatineeActor::Destroyed()
+void AMatineeActor::EndPlay(const EEndPlayReason::Type EndPlayReason)
 {
-	if( GetWorld() )
+	UWorld* World = GetWorld();
+	if (World)
 	{
-		GetWorldTimerManager().ClearTimer(TimerHandle_CheckPriorityRefresh);
+		World->GetTimerManager().ClearTimer(TimerHandle_CheckPriorityRefresh);
 	}
 
-	Super::Destroyed();
+	Super::EndPlay(EndPlayReason);
 }
 
 void AMatineeActor::SetPosition(float NewPosition,bool bJump)
@@ -1785,19 +1786,6 @@ void AMatineeActor::CheckPriorityRefresh()
 			return;
 		}
 	}
-
-	// check if it is controlling a player Pawn, or a platform a player Pawn is standing on
-	
-	/*foreach WorldSettings.AllControllers(class'Controller', C)
-	{
-		if ( C.PlayerState && C.Pawn != None &&
-			( InterpAction.LatentActors.Find(C.Pawn) != INDEX_NONE ||
-				(C.Pawn.Base != None && InterpAction.LatentActors.Find(C.Pawn.Base) != INDEX_NONE) ) )
-		{
-			ForceNetUpdate();
-			return;
-		}
-	}*/
 }
 
 
diff --git a/Engine/Source/Runtime/Engine/Private/Particles/ParticleComponents.cpp b/Engine/Source/Runtime/Engine/Private/Particles/ParticleComponents.cpp
--- a/Engine/Source/Runtime/Engine/Private/Particles/ParticleComponents.cpp
+++ b/Engine/Source/Runtime/Engine/Private/Particles/ParticleComponents.cpp
@@ -5935,13 +5935,13 @@ void AEmitterCameraLensEffectBase::UpdateLocation(const FVector& CamLoc, const F
 	SetActorLocationAndRotation( CamLoc + X * DistAdjustedForFOV, NewRot, false );
 }
 
-void AEmitterCameraLensEffectBase::Destroyed()
+void AEmitterCameraLensEffectBase::EndPlay(const EEndPlayReason::Type EndPlayReason)
 {
 	if (BaseCamera != NULL)
 	{
 		BaseCamera->RemoveCameraLensEffect(this);
 	}
-	Super::Destroyed();
+	Super::EndPlay(EndPlayReason);
 }
 
 void AEmitterCameraLensEffectBase::RegisterCamera(APlayerCameraManager* C)
@@ -5961,11 +5961,12 @@ void AEmitterCameraLensEffectBase::PostInitializeComponents()
 void AEmitterCameraLensEffectBase::ActivateLensEffect()
 {
 	// only play the camera effect on clients
-	check(GetWorld());
+	UWorld* World = GetWorld();
+	check(World);
 	if( GetNetMode() != NM_DedicatedServer )
 	{
 		UParticleSystem* PSToActuallySpawn;
-		if( GetWorld()->GameState && GetWorld()->GameState->ShouldShowGore() )
+		if( World->GameState && World->GameState->ShouldShowGore() )
 		{
 			PSToActuallySpawn = PS_CameraEffect;
 		}
diff --git a/Engine/Source/Runtime/Engine/Private/PlayerCameraManager.cpp b/Engine/Source/Runtime/Engine/Private/PlayerCameraManager.cpp
--- a/Engine/Source/Runtime/Engine/Private/PlayerCameraManager.cpp
+++ b/Engine/Source/Runtime/Engine/Private/PlayerCameraManager.cpp
@@ -762,15 +762,18 @@ void APlayerCameraManager::PostInitializeComponents()
 	AnimCameraActor = GetWorld()->SpawnActor<ACameraActor>(SpawnInfo);
 }
 
-void APlayerCameraManager::Destroyed()
+void APlayerCameraManager::EndPlay(const EEndPlayReason::Type EndPlayReason)
 {
 	// clean up the temp camera actor
 	if (AnimCameraActor)
 	{
-		AnimCameraActor->Destroy();
+		if (EndPlayReason == EEndPlayReason::Destroyed)
+		{
+			AnimCameraActor->Destroy();
+		}
 		AnimCameraActor = NULL;
 	}
-	Super::Destroyed();
+	Super::EndPlay(EndPlayReason);
 }
 
 void APlayerCameraManager::InitializeFor(APlayerController* PC)
diff --git a/Engine/Source/Runtime/Foliage/Private/InstancedFoliage.cpp b/Engine/Source/Runtime/Foliage/Private/InstancedFoliage.cpp
--- a/Engine/Source/Runtime/Foliage/Private/InstancedFoliage.cpp
+++ b/Engine/Source/Runtime/Foliage/Private/InstancedFoliage.cpp
@@ -1916,16 +1916,19 @@ bool AInstancedFoliageActor::HasSelectedInstances() const
 
 void AInstancedFoliageActor::Destroyed()
 {
-	for (auto& MeshPair : FoliageMeshes)
+	if (GIsEditor && !GetWorld()->IsGameWorld())
 	{
-		UHierarchicalInstancedStaticMeshComponent* Component = MeshPair.Value->Component;
-
-		if (Component)
+		for (auto& MeshPair : FoliageMeshes)
 		{
-			Component->ClearInstances();
-			// Save the component's PendingKill flag to restore the component if the delete is undone.
-			Component->SetFlags(RF_Transactional);
-			Component->Modify();
+			UHierarchicalInstancedStaticMeshComponent* Component = MeshPair.Value->Component;
+
+			if (Component)
+			{
+				Component->ClearInstances();
+				// Save the component's PendingKill flag to restore the component if the delete is undone.
+				Component->SetFlags(RF_Transactional);
+				Component->Modify();
+			}
 		}
 	}
 
@@ -2189,11 +2192,11 @@ void AInstancedFoliageActor::Serialize(FArchive& Ar)
 	}
 }
 
+#if WITH_EDITOR
 void AInstancedFoliageActor::PostInitProperties()
 {
 	Super::PostInitProperties();
 
-#if WITH_EDITOR
 	if (!IsTemplate())
 	{
 		GEngine->OnActorMoved().Remove(OnLevelActorMovedDelegateHandle);
@@ -2205,22 +2208,20 @@ void AInstancedFoliageActor::PostInitProperties()
 		FWorldDelegates::PostApplyLevelOffset.Remove(OnPostApplyLevelOffsetDelegateHandle);
 		OnPostApplyLevelOffsetDelegateHandle = FWorldDelegates::PostApplyLevelOffset.AddUObject(this, &AInstancedFoliageActor::OnPostApplyLevelOffset);
 	}
-#endif
 }
 
 void AInstancedFoliageActor::BeginDestroy()
 {
 	Super::BeginDestroy();
 
-#if WITH_EDITOR
 	if (!IsTemplate())
 	{
 		GEngine->OnActorMoved().Remove(OnLevelActorMovedDelegateHandle);
 		GEngine->OnLevelActorDeleted().Remove(OnLevelActorDeletedDelegateHandle);
 		FWorldDelegates::PostApplyLevelOffset.Remove(OnPostApplyLevelOffsetDelegateHandle);
 	}
-#endif
 }
+#endif
 
 void AInstancedFoliageActor::PostLoad()
 {
diff --git a/Engine/Source/Runtime/GameplayAbilities/Private/Abilities/GameplayAbilityTargetActor.cpp b/Engine/Source/Runtime/GameplayAbilities/Private/Abilities/GameplayAbilityTargetActor.cpp
--- a/Engine/Source/Runtime/GameplayAbilities/Private/Abilities/GameplayAbilityTargetActor.cpp
+++ b/Engine/Source/Runtime/GameplayAbilities/Private/Abilities/GameplayAbilityTargetActor.cpp
@@ -20,7 +20,7 @@ AGameplayAbilityTargetActor::AGameplayAbilityTargetActor(const FObjectInitialize
 	bDestroyOnConfirmation = true;
 }
 
-void AGameplayAbilityTargetActor::Destroyed()
+void AGameplayAbilityTargetActor::EndPlay(const EEndPlayReason::Type EndPlayReason)
 {
 	// We must remove ourselves from GenericLocalConfirmCallbacks/GenericLocalCancelCallbacks, since while these are bound they will inhibit any *other* abilities
 	// that are bound to the same key.
@@ -39,7 +39,7 @@ void AGameplayAbilityTargetActor::Destroyed()
 		}
 	}
 
-	Super::Destroyed();
+	Super::EndPlay(EndPlayReason);
 }
 
 void AGameplayAbilityTargetActor::GetLifetimeReplicatedProps(TArray< FLifetimeProperty > & OutLifetimeProps) const
diff --git a/Engine/Source/Runtime/GameplayAbilities/Public/Abilities/GameplayAbilityTargetActor.h b/Engine/Source/Runtime/GameplayAbilities/Public/Abilities/GameplayAbilityTargetActor.h
--- a/Engine/Source/Runtime/GameplayAbilities/Public/Abilities/GameplayAbilityTargetActor.h
+++ b/Engine/Source/Runtime/GameplayAbilities/Public/Abilities/GameplayAbilityTargetActor.h
@@ -24,7 +24,7 @@ class GAMEPLAYABILITIES_API AGameplayAbilityTargetActor : public AActor
 
 public:
 
-	void Destroyed() override;
+	void EndPlay(const EEndPlayReason::Type EndPlayReason) override;
 
 	/** The TargetData this class produces can be entirely generated on the server. We don't require the client to send us full or partial TargetData (possibly just a 'confirm') */
 	UPROPERTY(EditAnywhere, Category=Advanced)
diff --git a/Engine/Source/Runtime/Landscape/Private/Landscape.cpp b/Engine/Source/Runtime/Landscape/Private/Landscape.cpp
--- a/Engine/Source/Runtime/Landscape/Private/Landscape.cpp
+++ b/Engine/Source/Runtime/Landscape/Private/Landscape.cpp
@@ -1363,20 +1363,20 @@ void ALandscapeProxy::Destroyed()
 {
 	Super::Destroyed();
 
-	if (GIsEditor)
+	if (GIsEditor && !GetWorld()->IsGameWorld())
 	{
 		ULandscapeInfo::RecreateLandscapeInfo(GetWorld(), false);
-	}
 
-	if (SplineComponent)
-	{
-		SplineComponent->ModifySplines();
-	}
+		if (SplineComponent)
+		{
+			SplineComponent->ModifySplines();
+		}
 
-	TotalComponentsNeedingGrassMapRender -= NumComponentsNeedingGrassMapRender;
-	NumComponentsNeedingGrassMapRender = 0;
-	TotalTexturesToStreamForVisibleGrassMapRender -= NumTexturesToStreamForVisibleGrassMapRender;
-	NumTexturesToStreamForVisibleGrassMapRender = 0;
+		TotalComponentsNeedingGrassMapRender -= NumComponentsNeedingGrassMapRender;
+		NumComponentsNeedingGrassMapRender = 0;
+		TotalTexturesToStreamForVisibleGrassMapRender -= NumTexturesToStreamForVisibleGrassMapRender;
+		NumTexturesToStreamForVisibleGrassMapRender = 0;
+	}
 }
 
 void ALandscapeProxy::GetSharedProperties(ALandscapeProxy* Landscape)
