Fixed: Procedural Foliage Volumes affecting all sub-levels on simulate UE-18359

diff --git a/Engine/Source/Runtime/Foliage/Private/ProceduralFoliageComponent.cpp b/Engine/Source/Runtime/Foliage/Private/ProceduralFoliageComponent.cpp
--- a/Engine/Source/Runtime/Foliage/Private/ProceduralFoliageComponent.cpp
+++ b/Engine/Source/Runtime/Foliage/Private/ProceduralFoliageComponent.cpp
@@ -242,7 +242,6 @@ void UProceduralFoliageComponent::RemoveProceduralContent()
 			AInstancedFoliageActor* IFA = AInstancedFoliageActor::GetInstancedFoliageActorForLevel(Level);
 			if (IFA)
 			{
-				IFA->Modify();
 				IFA->DeleteInstancesForProceduralFoliageComponent(this);
 			}
 		}
