Map Check item to consolidate duplicate Instanced Foliage Actors

diff --git a/Engine/Source/Runtime/Foliage/Private/InstancedFoliage.cpp b/Engine/Source/Runtime/Foliage/Private/InstancedFoliage.cpp
--- a/Engine/Source/Runtime/Foliage/Private/InstancedFoliage.cpp
+++ b/Engine/Source/Runtime/Foliage/Private/InstancedFoliage.cpp
@@ -2240,11 +2240,18 @@ void AInstancedFoliageActor::PostLoad()
 	}
 	else
 	{
-		// Warn that there is more than one foliage actor in the scene
-		UE_LOG(LogInstancedFoliage, Warning, TEXT("Level %s: has more than one instanced foliage actor: %s, %s"), 
-			*OwningLevel->GetOutermost()->GetName(), 
-			*OwningLevel->InstancedFoliageActor->GetName(),
-			*this->GetName());
+		FFormatNamedArguments Arguments;
+		Arguments.Add(TEXT("Level"), FText::FromString(*OwningLevel->GetOutermost()->GetName()));
+		FMessageLog("MapCheck").Warning()
+			->AddToken(FUObjectToken::Create(this))
+			->AddToken(FTextToken::Create(FText::Format(LOCTEXT("MapCheck_DuplicateInstancedFoliageActor", "Level {Level} has an unexpected duplicate Instanced Foliage Actor {Actor}."), Arguments)))
+#if WITH_EDITOR
+			->AddToken(FActionToken::Create(LOCTEXT("MapCheck_FixDuplicateInstancedFoliageActor", "Fix"),
+				LOCTEXT("MapCheck_FixDuplicateInstancedFoliageActor_Desc", "Click consolidate into the main foliage actor."),
+				FOnActionTokenExecuted::CreateUObject(OwningLevel->InstancedFoliageActor.Get(), &AInstancedFoliageActor::RepairDuplicateIFA, this), true))
+#endif
+			;
+		FMessageLog("MapCheck").Open(EMessageSeverity::Warning);
 	}
 
 #if WITH_EDITOR
@@ -2348,6 +2355,30 @@ void AInstancedFoliageActor::PostLoad()
 
 #if WITH_EDITOR
 
+void AInstancedFoliageActor::RepairDuplicateIFA(AInstancedFoliageActor* DuplicateIFA)
+{
+	for (auto& MeshPair : DuplicateIFA->FoliageMeshes)
+	{
+		UFoliageType* DupeFoliageType = MeshPair.Key;
+		FFoliageMeshInfo& DupeMeshInfo = *MeshPair.Value;
+
+		// Get foliage type compatible with target IFA
+		FFoliageMeshInfo* TargetMeshInfo = nullptr;
+		UFoliageType* TargetFoliageType = AddFoliageType(DupeFoliageType, &TargetMeshInfo);
+
+		// Copy the instances
+		for (FFoliageInstance& Instance : DupeMeshInfo.Instances)
+		{
+			if ((Instance.Flags & FOLIAGE_InstanceDeleted) == 0)
+			{
+				TargetMeshInfo->AddInstance(this, TargetFoliageType, Instance);
+			}
+		}
+	}
+
+	GetWorld()->DestroyActor(DuplicateIFA);
+}
+
 void AInstancedFoliageActor::NotifyFoliageTypeChanged(UFoliageType* FoliageType, bool bMeshChanged)
 {
 	FFoliageMeshInfo* TypeInfo = FindMesh(FoliageType);
diff --git a/Engine/Source/Runtime/Foliage/Public/InstancedFoliageActor.h b/Engine/Source/Runtime/Foliage/Public/InstancedFoliageActor.h
--- a/Engine/Source/Runtime/Foliage/Public/InstancedFoliageActor.h
+++ b/Engine/Source/Runtime/Foliage/Public/InstancedFoliageActor.h
@@ -174,6 +174,8 @@ public:
 	DECLARE_EVENT_OneParam(AInstancedFoliageActor, FOnFoliageTypeMeshChanged, UFoliageType*);
 	FOnFoliageTypeMeshChanged& OnFoliageTypeMeshChanged() { return OnFoliageTypeMeshChangedEvent; }
 
+	/* Fix up a duplicate IFA */
+	void RepairDuplicateIFA(AInstancedFoliageActor* InDuplicateIFA);
 #endif	//WITH_EDITOR
 
 private:
