Visibility of each foliage type can be switched on/off per editor viewport UE-5238

diff --git a/Engine/Source/Editor/LevelEditor/LevelEditor.Build.cs b/Engine/Source/Editor/LevelEditor/LevelEditor.Build.cs
--- a/Engine/Source/Editor/LevelEditor/LevelEditor.Build.cs
+++ b/Engine/Source/Editor/LevelEditor/LevelEditor.Build.cs
@@ -59,7 +59,8 @@ public class LevelEditor : ModuleRules
                 "Persona",
                 "Kismet",
 				"KismetWidgets",
-				"Sequencer"
+				"Sequencer",
+                "Foliage"
 			}
 		);
 
diff --git a/Engine/Source/Editor/LevelEditor/Private/SLevelViewport.cpp b/Engine/Source/Editor/LevelEditor/Private/SLevelViewport.cpp
--- a/Engine/Source/Editor/LevelEditor/Private/SLevelViewport.cpp
+++ b/Engine/Source/Editor/LevelEditor/Private/SLevelViewport.cpp
@@ -40,6 +40,7 @@
 #include "GameFramework/PlayerController.h"
 #include "SActorPilotViewportToolbar.h"
 #include "SGameLayerManager.h"
+#include "FoliageType.h"
 
 static const FName LevelEditorName("LevelEditor");
 
@@ -1798,6 +1799,57 @@ bool SLevelViewport::IsLayerVisible( FName LayerName ) const
 	return LevelViewportClient->ViewHiddenLayers.Find(LayerName) == INDEX_NONE;
 }
 
+void SLevelViewport::ToggleShowFoliageType(TWeakObjectPtr<UFoliageType> InFoliageType)
+{
+	UFoliageType* FoliageType = InFoliageType.Get();
+	if (FoliageType)
+	{
+		FoliageType->HiddenEditorViews^= (1ull << LevelViewportClient->ViewIndex);
+		// Notify UFoliageType that things have changed
+		FoliageType->OnHiddenEditorViewMaskChanged(GetWorld());
+	
+		// Make sure to redraw viewport when user toggles foliage
+		LevelViewportClient->Invalidate();
+	}
+}
+
+void SLevelViewport::ToggleAllFoliageTypes(bool bVisible)
+{
+	UWorld* CurrentWorld = GetWorld(); 
+	TArray<UFoliageType*> AllFoliageTypes = GEditor->GetFoliageTypesInWorld(CurrentWorld);
+	if (AllFoliageTypes.Num())
+	{
+		const uint64 ViewMask = (1ull << LevelViewportClient->ViewIndex);
+
+		for (UFoliageType* FoliageType  : AllFoliageTypes)
+		{
+			if (bVisible)
+			{
+				FoliageType->HiddenEditorViews&= ~ViewMask;
+			}
+			else
+			{
+				FoliageType->HiddenEditorViews|= ViewMask;
+			}
+			
+			FoliageType->OnHiddenEditorViewMaskChanged(CurrentWorld);
+		}
+		
+		// Make sure to redraw viewport when user toggles meshes
+		LevelViewportClient->Invalidate(); 
+	}
+}
+
+bool SLevelViewport::IsFoliageTypeVisible(TWeakObjectPtr<UFoliageType> InFoliageType) const
+{
+	const UFoliageType* FoliageType = InFoliageType.Get();
+	if (FoliageType)
+	{
+		return (FoliageType->HiddenEditorViews & (1ull << LevelViewportClient->ViewIndex)) == 0;
+	}
+	return false;
+}
+
 FViewport* SLevelViewport::GetActiveViewport() 
 { 
 	return ActiveViewport->GetViewport(); 
diff --git a/Engine/Source/Editor/LevelEditor/Private/SLevelViewportToolBar.cpp b/Engine/Source/Editor/LevelEditor/Private/SLevelViewportToolBar.cpp
--- a/Engine/Source/Editor/LevelEditor/Private/SLevelViewportToolBar.cpp
+++ b/Engine/Source/Editor/LevelEditor/Private/SLevelViewportToolBar.cpp
@@ -22,6 +22,7 @@
 #include "StatsData.h"
 #include "BufferVisualizationData.h"
 #include "GameFramework/WorldSettings.h"
+#include "FoliageType.h"
 
 #define LOCTEXT_NAMESPACE "LevelViewportToolBar"
 
@@ -938,9 +939,9 @@ TSharedRef<SWidget> SLevelViewportToolBar::GenerateShowMenu() const
 
 			// Show Layers sub-menu is dynamically generated when the user enters 'show' menu
 			{
-			ShowMenuBuilder.AddSubMenu(LOCTEXT("ShowLayersMenu", "Layers"), LOCTEXT("ShowLayersMenu_ToolTip", "Show layers flags"),
-				FNewMenuDelegate::CreateStatic(&SLevelViewportToolBar::FillShowLayersMenu, Viewport));
-		}
+				ShowMenuBuilder.AddSubMenu(LOCTEXT("ShowLayersMenu", "Layers"), LOCTEXT("ShowLayersMenu_ToolTip", "Show layers flags"),
+					FNewMenuDelegate::CreateStatic(&SLevelViewportToolBar::FillShowLayersMenu, Viewport));
+			}
 
 			// Show Sprites sub-menu
 			{
@@ -956,6 +957,13 @@ TSharedRef<SWidget> SLevelViewportToolBar::GenerateShowMenu() const
 				ShowMenuBuilder.AddSubMenu(LOCTEXT("ShowSpritesMenu", "Sprites"), LOCTEXT("ShowSpritesMenu_ToolTip", "Show sprites flags"),
 					FNewMenuDelegate::CreateStatic(&FillShowMenu, ShowSpritesMenu, 2));
 			}
+
+			// Show 'Foliage types' sub-menu is dynamically generated when the user enters 'show' menu
+			{
+				ShowMenuBuilder.AddSubMenu(LOCTEXT("ShowFoliageTypesMenu", "Foliage Types"), LOCTEXT("ShowFoliageTypesMenu_ToolTip", "Show/hide specific foliage types"),
+					FNewMenuDelegate::CreateStatic(&SLevelViewportToolBar::FillShowFoliageTypesMenu, Viewport));
+			}
+
 		}
 
 		ShowMenuBuilder.EndSection();
@@ -1081,6 +1089,47 @@ void SLevelViewportToolBar::FillShowLayersMenu( FMenuBuilder& MenuBuilder, TWeak
 	}
 }
 
+void SLevelViewportToolBar::FillShowFoliageTypesMenu(class FMenuBuilder& MenuBuilder, TWeakPtr<class SLevelViewport> Viewport)
+{
+	auto ViewportPtr = Viewport.Pin();
+	if (!ViewportPtr.IsValid())
+	{
+		return;
+	}
+		
+	MenuBuilder.BeginSection("LevelViewportFoliageMeshes");
+	{
+		// Map 'Show All' and 'Hide All' commands
+		FUIAction ShowAllFoliage(FExecuteAction::CreateSP(ViewportPtr.ToSharedRef(), &SLevelViewport::ToggleAllFoliageTypes, true));
+		FUIAction HideAllFoliage(FExecuteAction::CreateSP(ViewportPtr.ToSharedRef(), &SLevelViewport::ToggleAllFoliageTypes, false));
+		
+		MenuBuilder.AddMenuEntry(LOCTEXT("ShowAllLabel", "Show All"), FText::GetEmpty(), FSlateIcon(), ShowAllFoliage);
+		MenuBuilder.AddMenuEntry(LOCTEXT("HideAllLabel", "Hide All"), FText::GetEmpty(), FSlateIcon(), HideAllFoliage);
+	}
+	MenuBuilder.EndSection();
+		
+	TArray<UFoliageType*> AllFoliageTypes = GEditor->GetFoliageTypesInWorld(ViewportPtr->GetWorld());
+
+	if (AllFoliageTypes.Num())
+	{
+		MenuBuilder.BeginSection("LevelViewportFoliageTypes");
+			
+		for (UFoliageType* FoliageType : AllFoliageTypes)
+		{
+			FString MeshName = FoliageType->GetName();
+			TWeakObjectPtr<UFoliageType> FoliageTypePtr = FoliageType;
+				
+			FUIAction Action(
+					FExecuteAction::CreateSP(ViewportPtr.ToSharedRef(), &SLevelViewport::ToggleShowFoliageType, FoliageTypePtr),
+					FCanExecuteAction(),
+					FIsActionChecked::CreateSP(ViewportPtr.ToSharedRef(), &SLevelViewport::IsFoliageTypeVisible, FoliageTypePtr));
+
+			MenuBuilder.AddMenuEntry(FText::FromString(MeshName), FText::GetEmpty(), FSlateIcon(), Action, NAME_None, EUserInterfaceActionType::ToggleButton);
+		}
+
+		MenuBuilder.EndSection();
+	}
+}
 
 TWeakObjectPtr<UWorld> SLevelViewportToolBar::GetWorld() const
 {
diff --git a/Engine/Source/Editor/LevelEditor/Private/SLevelViewportToolBar.h b/Engine/Source/Editor/LevelEditor/Private/SLevelViewportToolBar.h
--- a/Engine/Source/Editor/LevelEditor/Private/SLevelViewportToolBar.h
+++ b/Engine/Source/Editor/LevelEditor/Private/SLevelViewportToolBar.h
@@ -156,6 +156,14 @@ private:
 	 */
 	static void FillShowLayersMenu( class FMenuBuilder& MenuBuilder, TWeakPtr<class SLevelViewport> Viewport );
 
+	/**
+	 * Generates 'show foliage types' menu content for a viewport
+	 *
+	 * @param MenuBuilder	menu builder
+	 * @param Viewport		target vieport
+	 */
+	static void FillShowFoliageTypesMenu(class FMenuBuilder& MenuBuilder, TWeakPtr<class SLevelViewport> Viewport);
+
 	/** Generates the layout sub-menu content */
 	void GenerateViewportConfigsMenu(FMenuBuilder& MenuBuilder) const;
 
diff --git a/Engine/Source/Editor/LevelEditor/Public/SLevelViewport.h b/Engine/Source/Editor/LevelEditor/Public/SLevelViewport.h
--- a/Engine/Source/Editor/LevelEditor/Public/SLevelViewport.h
+++ b/Engine/Source/Editor/LevelEditor/Public/SLevelViewport.h
@@ -174,6 +174,27 @@ public:
 	 */
 	bool IsLayerVisible( FName LayerName ) const;
 
+	/**
+	 * Toggles foliage type visibility in this viewport
+	 *
+	 * @param FoliageType	Target foliage type
+	 */
+	void ToggleShowFoliageType(TWeakObjectPtr<class UFoliageType> FoliageType);
+
+	/**
+	 * Toggles all foliage types visibility
+	 *
+	 * @param Visible	true if foliage types should be visible, false otherwise
+	 */
+	void ToggleAllFoliageTypes(bool bVisible);
+
+	/**
+	 * Checks if a foliage type is visible in this viewport
+	 *
+	 * @param FoliageType	Target foliage type
+	 */
+	bool IsFoliageTypeVisible(TWeakObjectPtr<class UFoliageType> FoliageType) const;
+
 
 	/** Called to lock/unlock the actor from the viewport's context menu */
 	void OnActorLockToggleFromMenu(AActor* Actor);
diff --git a/Engine/Source/Editor/UnrealEd/Classes/Editor/EditorEngine.h b/Engine/Source/Editor/UnrealEd/Classes/Editor/EditorEngine.h
--- a/Engine/Source/Editor/UnrealEd/Classes/Editor/EditorEngine.h
+++ b/Engine/Source/Editor/UnrealEd/Classes/Editor/EditorEngine.h
@@ -16,6 +16,8 @@ class FAssetData;
 class FPoly;
 class UAnimSequence;
 class USkeleton;
+class UFoliageType;
+
 //
 // Things to set in mapSetBrush.
 //
@@ -1792,6 +1794,14 @@ public:
 	 */
 	void MoveSelectedFoliageToLevel( ULevel* InTargetLevel );
 
+	/**
+	 *	Returns list of all foliage types used in the world
+	 * 
+	 * @param	InWorld	 The target world.
+	 * @return	List of all foliage types used in the world
+	 */
+	TArray<UFoliageType*> GetFoliageTypesInWorld(UWorld* InWorld);
+
 	/**
 	 * Checks to see whether it's possible to perform a copy operation on the selected actors.
 	 *
diff --git a/Engine/Source/Editor/UnrealEd/Private/EditorServer.cpp b/Engine/Source/Editor/UnrealEd/Private/EditorServer.cpp
--- a/Engine/Source/Editor/UnrealEd/Private/EditorServer.cpp
+++ b/Engine/Source/Editor/UnrealEd/Private/EditorServer.cpp
@@ -2962,6 +2962,22 @@ void UEditorEngine::MoveSelectedFoliageToLevel(ULevel* InTargetLevel)
 	}
 }
 
+TArray<UFoliageType*> UEditorEngine::GetFoliageTypesInWorld(UWorld* InWorld)
+{
+	TSet<UFoliageType*> FoliageSet;
+	
+	// Iterate over all foliage actors in the world
+	for (TActorIterator<AInstancedFoliageActor> It(InWorld); It; ++It)
+	{
+		for (const auto& Pair : It->FoliageMeshes)
+		{
+			FoliageSet.Add(Pair.Key);
+		}
+	}
+
+	return FoliageSet.Array();
+}
+
 ULevel*  UEditorEngine::CreateTransLevelMoveBuffer( UWorld* InWorld )
 {
 	ULevel* BufferLevel = NewObject<ULevel>(GetTransientPackage(), TEXT("TransLevelMoveBuffer"));
diff --git a/Engine/Source/Runtime/Engine/Classes/Components/PrimitiveComponent.h b/Engine/Source/Runtime/Engine/Classes/Components/PrimitiveComponent.h
--- a/Engine/Source/Runtime/Engine/Classes/Components/PrimitiveComponent.h
+++ b/Engine/Source/Runtime/Engine/Classes/Components/PrimitiveComponent.h
@@ -1506,6 +1506,11 @@ public:
 	
 	virtual bool GetShadowIndirectOnly() const { return false; }
 
+#if WITH_EDITOR
+	/** Returns mask that represents in which views this primitive is hidden */
+	virtual uint64 GetHiddenEditorViews() const;
+#endif// WITH_EDITOR
+
 	/**
 	 *	Set the angular velocity of all bodies in this component.
 	 *
diff --git a/Engine/Source/Runtime/Engine/Private/Components/PrimitiveComponent.cpp b/Engine/Source/Runtime/Engine/Private/Components/PrimitiveComponent.cpp
--- a/Engine/Source/Runtime/Engine/Private/Components/PrimitiveComponent.cpp
+++ b/Engine/Source/Runtime/Engine/Private/Components/PrimitiveComponent.cpp
@@ -986,6 +986,13 @@ void UPrimitiveComponent::PushEditorVisibilityToProxy( uint64 InVisibility )
 	}
 }
 
+#if WITH_EDITOR
+uint64 UPrimitiveComponent::GetHiddenEditorViews() const
+{
+	const AActor* OwnerActor = GetOwner();
+	return OwnerActor ? OwnerActor->HiddenEditorViews : 0;
+}
+#endif// WITH_EDITOR
 
 void UPrimitiveComponent::PushHoveredToProxy(const bool bInHovered)
 {
diff --git a/Engine/Source/Runtime/Engine/Private/PrimitiveSceneProxy.cpp b/Engine/Source/Runtime/Engine/Private/PrimitiveSceneProxy.cpp
--- a/Engine/Source/Runtime/Engine/Private/PrimitiveSceneProxy.cpp
+++ b/Engine/Source/Runtime/Engine/Private/PrimitiveSceneProxy.cpp
@@ -131,7 +131,7 @@ FPrimitiveSceneProxy::FPrimitiveSceneProxy(const UPrimitiveComponent* InComponen
 
 #if WITH_EDITOR
 		// cache the actor's group membership
-		HiddenEditorViews = InComponent->GetOwner()->HiddenEditorViews;
+		HiddenEditorViews = InComponent->GetHiddenEditorViews();
 #endif
 	}
 	
diff --git a/Engine/Source/Runtime/Foliage/Private/InstancedFoliage.cpp b/Engine/Source/Runtime/Foliage/Private/InstancedFoliage.cpp
--- a/Engine/Source/Runtime/Foliage/Private/InstancedFoliage.cpp
+++ b/Engine/Source/Runtime/Foliage/Private/InstancedFoliage.cpp
@@ -52,6 +52,8 @@ struct FFoliageCustomVersion
 		FoliageHISMCBlueprints = 9,
 		// Added Mobility setting to UFoliageType
 		AddedMobility = 10,
+		// Make sure that foliage has FoliageHISMC class
+		FoliageUsingFoliageISMC = 11,
 		// -----<new versions can be added above this line>-------------------------------------------------
 		VersionPlusOne,
 		LatestVersion = VersionPlusOne - 1
@@ -328,6 +330,7 @@ UFoliageType::UFoliageType(const FObjectInitializer& ObjectInitializer)
 	Curve->AddKey(1.f, 1.f);
 
 	UpdateGuid = FGuid::NewGuid();
+	HiddenEditorViews = 0;
 
 #if WITH_EDITORONLY_DATA
 	// Deprecated since FFoliageCustomVersion::FoliageTypeCustomization
@@ -532,6 +535,25 @@ void UFoliageType::PostEditChangeProperty(struct FPropertyChangedEvent& Property
 		}
 	}
 }
+
+void UFoliageType::OnHiddenEditorViewMaskChanged(UWorld* InWorld)
+{
+	for (TActorIterator<AInstancedFoliageActor> It(InWorld); It; ++It)
+	{
+		FFoliageMeshInfo* MeshInfo = It->FindMesh(this);
+		if (MeshInfo && MeshInfo->Component)
+		{
+	 		UFoliageInstancedStaticMeshComponent* FoliageComponent = Cast<UFoliageInstancedStaticMeshComponent>(MeshInfo->Component);
+
+			if (FoliageComponent && FoliageComponent->FoliageHiddenEditorViews != HiddenEditorViews)
+			{
+				FoliageComponent->FoliageHiddenEditorViews = HiddenEditorViews;
+				FoliageComponent->MarkRenderStateDirty();
+			}
+		}
+	}
+}
+
 #endif
 
 //
@@ -629,13 +651,18 @@ void FFoliageMeshInfo::CreateNewComponent(AInstancedFoliageActor* InIFA, const U
 		ComponentClass = UFoliageInstancedStaticMeshComponent::StaticClass();
 	}
 
-	Component = NewObject<UFoliageInstancedStaticMeshComponent>(InIFA, ComponentClass, NAME_None, RF_Transactional);
+	UFoliageInstancedStaticMeshComponent* FoliageComponent = NewObject<UFoliageInstancedStaticMeshComponent>(InIFA, ComponentClass, NAME_None, RF_Transactional);
 
+	Component = FoliageComponent;
 	Component->StaticMesh = InSettings->GetStaticMesh();
 	Component->bSelectable = true;
 	Component->bHasPerInstanceHitProxies = true;
 	Component->InstancingRandomSeed = FMath::Rand();
 
+#if WITH_EDITOR
+	FoliageComponent->FoliageHiddenEditorViews = InSettings->HiddenEditorViews;
+#endif	
+
 	UpdateComponentSettings(InSettings);
 
 	Component->AttachTo(InIFA->GetRootComponent());
@@ -781,6 +808,14 @@ void FFoliageMeshInfo::UpdateComponentSettings(const UFoliageType* InSettings)
 			bNeedsInvalidateLightingCache = true;
 		}
 
+		UFoliageInstancedStaticMeshComponent* FoliageComponent = Cast<UFoliageInstancedStaticMeshComponent>(Component);
+
+		if (FoliageComponent && FoliageComponent->FoliageHiddenEditorViews != InSettings->HiddenEditorViews)
+		{
+			FoliageComponent->FoliageHiddenEditorViews = InSettings->HiddenEditorViews;
+			bNeedsMarkRenderStateDirty = true;
+		}
+		
 		const UFoliageType_InstancedStaticMesh* FoliageType_ISM = Cast<UFoliageType_InstancedStaticMesh>(FoliageType);
 		if (FoliageType_ISM)
 		{
@@ -1053,6 +1088,7 @@ void FFoliageMeshInfo::ReapplyInstancesToComponent()
 		// clear the transactional flag if it was set prior to deleting the actor
 		Component->ClearFlags(RF_Transactional);
 
+		const bool bWasRegistered = Component->IsRegistered();
 		Component->UnregisterComponent();
 		Component->ClearInstances();
 
@@ -1071,7 +1107,10 @@ void FFoliageMeshInfo::ReapplyInstancesToComponent()
 			}
 		}
 
-		Component->RegisterComponent();
+		if (bWasRegistered)
+		{
+			Component->RegisterComponent();
+		}
 	}
 }
 
@@ -2249,7 +2288,7 @@ void AInstancedFoliageActor::PostLoad()
 			->AddToken(FActionToken::Create(LOCTEXT("MapCheck_FixDuplicateInstancedFoliageActor", "Fix"),
 				LOCTEXT("MapCheck_FixDuplicateInstancedFoliageActor_Desc", "Click consolidate into the main foliage actor."),
 				FOnActionTokenExecuted::CreateUObject(OwningLevel->InstancedFoliageActor.Get(), &AInstancedFoliageActor::RepairDuplicateIFA, this), true))
-#endif
+#endif// WITH_EDITOR
 			;
 		FMessageLog("MapCheck").Open(EMessageSeverity::Warning);
 	}
@@ -2304,17 +2343,6 @@ void AInstancedFoliageActor::PostLoad()
 				MeshInfo.ReallocateClusters(this, MeshPair.Key);
 			}
 
-			// Update foliage component settings if the foliage settings object was changed while the level was not loaded.
-			if (MeshInfo.FoliageTypeUpdateGuid != FoliageType->UpdateGuid)
-			{
-				if (MeshInfo.FoliageTypeUpdateGuid.IsValid())
-				{
-					MeshInfo.CheckComponentClass(this, FoliageType);
-					MeshInfo.UpdateComponentSettings(FoliageType);
-				}
-				MeshInfo.FoliageTypeUpdateGuid = FoliageType->UpdateGuid;
-			}
-
 			// Update the hash.
 			MeshInfo.ComponentHash.Empty();
 			MeshInfo.InstanceHash->Empty();
@@ -2337,20 +2365,37 @@ void AInstancedFoliageActor::PostLoad()
 					MeshInfo.Component->ClearFlags(RF_Transactional);
 				}
 			}
-
-			//Clean up case where embeded instances had their static mesh deleted
+			
+			// Clean up case where embeded instances had their static mesh deleted
 			if (FoliageType->IsNotAssetOrBlueprint() && FoliageType->GetStaticMesh() == nullptr)
 			{
 				OnFoliageTypeMeshChangedEvent.Broadcast(FoliageType);
 				RemoveFoliageType(&FoliageType, 1);
+				continue;
+			}
+
+			// Upgrade foliage component
+			if (GetLinkerCustomVersion(FFoliageCustomVersion::GUID) < FFoliageCustomVersion::FoliageUsingFoliageISMC)
+			{
+				MeshInfo.CheckComponentClass(this, FoliageType);
+			}
+
+			// Update foliage component settings if the foliage settings object was changed while the level was not loaded.
+			if (MeshInfo.FoliageTypeUpdateGuid != FoliageType->UpdateGuid)
+			{
+				if (MeshInfo.FoliageTypeUpdateGuid.IsValid())
+				{
+					MeshInfo.CheckComponentClass(this, FoliageType);
+					MeshInfo.UpdateComponentSettings(FoliageType);
+				}
+				MeshInfo.FoliageTypeUpdateGuid = FoliageType->UpdateGuid;
 			}
 		}
 
 		// Clean up dead cross-level references
 		FFoliageInstanceBaseCache::CompactInstanceBaseCache(this);
-
 	}
-#endif
+#endif// WITH_EDITOR
 }
 
 #if WITH_EDITOR
@@ -2757,4 +2802,13 @@ void UFoliageInstancedStaticMeshComponent::ReceiveComponentDamage(float DamageAm
 	}
 }
 
+#if WITH_EDITOR
+
+uint64 UFoliageInstancedStaticMeshComponent::GetHiddenEditorViews() const
+{
+	return FoliageHiddenEditorViews;
+}
+
+#endif// WITH_EDITOR
+
 #undef LOCTEXT_NAMESPACE
diff --git a/Engine/Source/Runtime/Foliage/Public/FoliageInstancedStaticMeshComponent.h b/Engine/Source/Runtime/Foliage/Public/FoliageInstancedStaticMeshComponent.h
--- a/Engine/Source/Runtime/Foliage/Public/FoliageInstancedStaticMeshComponent.h
+++ b/Engine/Source/Runtime/Foliage/Public/FoliageInstancedStaticMeshComponent.h
@@ -22,6 +22,16 @@ public:
 	UPROPERTY(BlueprintAssignable, Category = "Game|Damage")
 	FInstanceRadialDamageSignature OnInstanceTakeRadialDamage;
 
+#if WITH_EDITORONLY_DATA
+	UPROPERTY(transient)
+	uint64 FoliageHiddenEditorViews;
+#endif// WITH_EDITOR
+
 	virtual void ReceiveComponentDamage(float DamageAmount, FDamageEvent const& DamageEvent, AController* EventInstigator, AActor* DamageCauser) override;
+
+	//USceneComponent Interface
+#if WITH_EDITOR
+	virtual uint64 GetHiddenEditorViews() const override;
+#endif// WITH_EDITOR
 };
 
diff --git a/Engine/Source/Runtime/Foliage/Public/FoliageType.h b/Engine/Source/Runtime/Foliage/Public/FoliageType.h
--- a/Engine/Source/Runtime/Foliage/Public/FoliageType.h
+++ b/Engine/Source/Runtime/Foliage/Public/FoliageType.h
@@ -58,6 +58,9 @@ class UFoliageType : public UObject
 	virtual bool IsFoliageReallocationRequiredForPropertyChange(struct FPropertyChangedEvent& PropertyChangedEvent) const { return true; }
 
 	virtual void PostEditChangeProperty(struct FPropertyChangedEvent& PropertyChangedEvent) override;
+
+	/* Notifies all relevant foliage actors that HiddenEditorView mask has been changed */
+	FOLIAGE_API void OnHiddenEditorViewMaskChanged(UWorld* InWorld);
 #endif
 
 	/* A GUID that is updated every time the foliage type is modified, 
@@ -235,6 +238,10 @@ public:
 	UPROPERTY(EditAnywhere, Category=InstanceSettings, meta=(HideObjectType=true))
 	TEnumAsByte<EHasCustomNavigableGeometry::Type> CustomNavigableGeometry;
 
+	/** Bitflag to represent in which editor views this foliage mesh is hidden. */
+	UPROPERTY(transient)
+	uint64 HiddenEditorViews;
+
 	UPROPERTY()
 	uint32 IsSelected:1;
 
