Fixed instance highlight selection remaining after undo (UE-19385)

diff --git a/Engine/Source/Runtime/Foliage/Private/InstancedFoliage.cpp b/Engine/Source/Runtime/Foliage/Private/InstancedFoliage.cpp
--- a/Engine/Source/Runtime/Foliage/Private/InstancedFoliage.cpp
+++ b/Engine/Source/Runtime/Foliage/Private/InstancedFoliage.cpp
@@ -1117,10 +1117,10 @@ void FFoliageMeshInfo::ReapplyInstancesToComponent()
 			Component->AddInstanceWorldSpace(Instance.GetInstanceWorldTransform());
 		}
 
+		Component->ClearInstanceSelection();
+
 		if (SelectedIndices.Num())
-		{
-			Component->ClearInstanceSelection();
-			
+		{	
 			for (int32 i : SelectedIndices)
 			{
 				Component->SelectInstance(true, i, 1);
