Added bUseAsOccluder checkbox to foliage tool to allow foliage to correctly receive DBuffer decals.

diff --git a/Engine/Source/Runtime/Foliage/Private/InstancedFoliage.cpp b/Engine/Source/Runtime/Foliage/Private/InstancedFoliage.cpp
--- a/Engine/Source/Runtime/Foliage/Private/InstancedFoliage.cpp
+++ b/Engine/Source/Runtime/Foliage/Private/InstancedFoliage.cpp
@@ -306,6 +306,7 @@ UFoliageType::UFoliageType(const FObjectInitializer& ObjectInitializer)
 
 	bOverrideLightMapRes = false;
 	OverriddenLightMapRes = 8;
+	bUseAsOccluder = false;
 
 	BodyInstance.SetCollisionProfileName(UCollisionProfile::NoCollision_ProfileName);
 
@@ -827,6 +828,11 @@ void FFoliageMeshInfo::UpdateComponentSettings(const UFoliageType* InSettings)
 			bNeedsMarkRenderStateDirty = true;
 			bNeedsInvalidateLightingCache = true;
 		}
+		if (Component->bUseAsOccluder != FoliageType->bUseAsOccluder)
+		{
+			Component->bUseAsOccluder = FoliageType->bUseAsOccluder;
+			bNeedsMarkRenderStateDirty = true;
+		}
 
 		UFoliageInstancedStaticMeshComponent* FoliageComponent = Cast<UFoliageInstancedStaticMeshComponent>(Component);
 
diff --git a/Engine/Source/Runtime/Foliage/Public/FoliageType.h b/Engine/Source/Runtime/Foliage/Public/FoliageType.h
--- a/Engine/Source/Runtime/Foliage/Public/FoliageType.h
+++ b/Engine/Source/Runtime/Foliage/Public/FoliageType.h
@@ -233,6 +233,13 @@ public:
 	UPROPERTY(EditAnywhere, BlueprintReadOnly, Category=InstanceSettings, meta=(DisplayName="Light Map Resolution", EditCondition="bOverrideLightMapRes"))
 	int32 OverriddenLightMapRes;
 
+	/**
+	 * If enabled, foliage will render a pre-pass which allows it to occlude other primitives, and also allows 
+	 * it to correctly receive DBuffer decals. Enabling this setting may have a negative performance impact.
+	 */
+	UPROPERTY(EditAnywhere, AdvancedDisplay, BlueprintReadOnly, Category = InstanceSettings)
+	uint32 bUseAsOccluder : 1;
+
 	/** Custom collision for foliage */
 	UPROPERTY(EditAnywhere, Category=InstanceSettings, meta=(HideObjectType=true))
 	struct FBodyInstance BodyInstance;
