enabled custom export of navigable geometry for foliage with InstancedStaticMesh type
#ue4

diff --git a/Engine/Source/Runtime/Foliage/Private/InstancedFoliage.cpp b/Engine/Source/Runtime/Foliage/Private/InstancedFoliage.cpp
--- a/Engine/Source/Runtime/Foliage/Private/InstancedFoliage.cpp
+++ b/Engine/Source/Runtime/Foliage/Private/InstancedFoliage.cpp
@@ -467,6 +467,7 @@ UFoliageType_InstancedStaticMesh::UFoliageType_InstancedStaticMesh(const FObject
 {
 	Mesh = nullptr;
 	ComponentClass = UFoliageInstancedStaticMeshComponent::StaticClass();
+	CustomNavigableGeometry = EHasCustomNavigableGeometry::Yes;
 }
 
 
