Fixed crash painting foliage on BSP

diff --git a/Engine/Source/Runtime/Foliage/Private/InstancedFoliage.cpp b/Engine/Source/Runtime/Foliage/Private/InstancedFoliage.cpp
--- a/Engine/Source/Runtime/Foliage/Private/InstancedFoliage.cpp
+++ b/Engine/Source/Runtime/Foliage/Private/InstancedFoliage.cpp
@@ -2647,13 +2647,13 @@ bool AInstancedFoliageActor::FoliageTrace(const UWorld* InWorld, FHitResult& Out
 
 	for (const FHitResult& Hit : Hits)
 	{
-		const UPrimitiveComponent* HitComponent = Hit.GetComponent();
+		const AActor* HitActor = Hit.GetActor();
 
 		// don't place procedural foliage inside an AProceduralFoliageBlockingVolume
 		// this test is first because two of the tests below would otherwise cause the trace to ignore AProceduralFoliageBlockingVolume
 		if (DesiredInstance.PlacementMode == EFoliagePlacementMode::Procedural)
 		{
-			if (const AProceduralFoliageBlockingVolume* ProceduralFoliageBlockingVolume = Cast<AProceduralFoliageBlockingVolume>(Hit.Actor.Get()))
+			if (const AProceduralFoliageBlockingVolume* ProceduralFoliageBlockingVolume = Cast<AProceduralFoliageBlockingVolume>(HitActor))
 			{
 				const AProceduralFoliageVolume* ProceduralFoliageVolume = ProceduralFoliageBlockingVolume->ProceduralFoliageVolume;
 				if (ProceduralFoliageVolume == nullptr || ProceduralFoliageVolume->ProceduralComponent == nullptr || ProceduralFoliageVolume->ProceduralComponent->GetProceduralGuid() == DesiredInstance.ProceduralGuid)
@@ -2661,18 +2661,21 @@ bool AInstancedFoliageActor::FoliageTrace(const UWorld* InWorld, FHitResult& Out
 					return false;
 				}
 			}
-			else if (Hit.GetActor()->IsA<AProceduralFoliageVolume>()) //we never want to collide with our spawning volume
+			else if (HitActor && HitActor->IsA<AProceduralFoliageVolume>()) //we never want to collide with our spawning volume
 			{
 				continue;
 			}
 		}
 
 		// Don't place foliage on foliage
-		if (Hit.GetActor()->IsA<AInstancedFoliageActor>())
+		if (HitActor && HitActor->IsA<AInstancedFoliageActor>())
 		{
 			return false;
 		}
 
+		const UPrimitiveComponent* HitComponent = Hit.GetComponent();
+		check(HitComponent);
+
 		// In the editor traces can hit "No Collision" type actors, so ugh. (ignore these)
 		const FBodyInstance* BodyInstance = HitComponent->GetBodyInstance();
 		if (BodyInstance->GetCollisionEnabled() != ECollisionEnabled::QueryAndPhysics || BodyInstance->GetResponseToChannel(ECC_WorldStatic) != ECR_Block)
