Restrict the Paint/Unpaint Density to 1.0
Restrict the Foliage Density type to 10000
Restrict the Grass Density type to 1000
Fixed a NAN that could be generated when adding instances
These settings would still allow to paint ~2 millions instances with the biggest brush and max density per click.

diff --git a/Engine/Source/Editor/FoliageEdit/Private/FoliageEdMode.cpp b/Engine/Source/Editor/FoliageEdit/Private/FoliageEdMode.cpp
--- a/Engine/Source/Editor/FoliageEdit/Private/FoliageEdMode.cpp
+++ b/Engine/Source/Editor/FoliageEdit/Private/FoliageEdMode.cpp
@@ -1078,7 +1078,7 @@ void FEdModeFoliage::GetRandomVectorInBrush(FVector& OutStart, FVector& OutEnd)
 	FVector Point = Ru * U + Rv * V;
 
 	// find distance to surface of sphere brush from this point
-	FVector Rw = FMath::Sqrt(1.f - (FMath::Square(Ru) + FMath::Square(Rv))) * BrushNormal;
+	FVector Rw = FMath::Sqrt(FMath::Max(1.f - (FMath::Square(Ru) + FMath::Square(Rv)), 0.001f)) * BrushNormal;
 
 	OutStart = BrushLocation + UISettings.GetRadius() * (Point + Rw);
 	OutEnd = BrushLocation + UISettings.GetRadius() * (Point - Rw);
diff --git a/Engine/Source/Editor/FoliageEdit/Private/SFoliageEdit.cpp b/Engine/Source/Editor/FoliageEdit/Private/SFoliageEdit.cpp
--- a/Engine/Source/Editor/FoliageEdit/Private/SFoliageEdit.cpp
+++ b/Engine/Source/Editor/FoliageEdit/Private/SFoliageEdit.cpp
@@ -184,7 +184,7 @@ void SFoliageEdit::Construct(const FArguments& InArgs)
 							.Font(StandardFont)
 							.AllowSpin(true)
 							.MinValue(0.0f)
-							.MaxValue(2.0f)
+							.MaxValue(1.0f)
 							.MaxSliderValue(1.0f)
 							.MinDesiredValueWidth(50.0f)
 							.Value(this, &SFoliageEdit::GetPaintDensity)
@@ -220,7 +220,7 @@ void SFoliageEdit::Construct(const FArguments& InArgs)
 							.Font(StandardFont)
 							.AllowSpin(true)
 							.MinValue(0.0f)
-							.MaxValue(2.0f)
+							.MaxValue(1.0f)
 							.MaxSliderValue(1.0f)
 							.MinDesiredValueWidth(50.0f)
 							.Value(this, &SFoliageEdit::GetEraseDensity)
diff --git a/Engine/Source/Runtime/Foliage/Public/FoliageType.h b/Engine/Source/Runtime/Foliage/Public/FoliageType.h
--- a/Engine/Source/Runtime/Foliage/Public/FoliageType.h
+++ b/Engine/Source/Runtime/Foliage/Public/FoliageType.h
@@ -121,15 +121,15 @@ public:
 	// PAINTING
 
 	/** Foliage instances will be placed at this density, specified in instances per 1000x1000 unit area */
-	UPROPERTY(EditAnywhere, Category=Painting, meta=(DisplayName="Density / 1Kuu", UIMin = 0, ClampMin = 0, UIMax = 100000, ClampMax = 100000))
+	UPROPERTY(EditAnywhere, Category=Painting, meta=(DisplayName="Density / 1Kuu", UIMin = 0, ClampMin = 0, UIMax = 10000, ClampMax = 10000))
 	float Density;
 
 	/** The factor by which to adjust the density of instances. Values >1 will increase density while values <1 will decrease it. */
-	UPROPERTY(EditAnywhere, Category=Painting, meta=(UIMin=0, ClampMin=0, UIMax = 10000, ClampMax = 10000, ReapplyCondition="ReapplyDensity"))
+	UPROPERTY(EditAnywhere, Category=Painting, meta=(UIMin=0, ClampMin=0, UIMax = 1000, ClampMax = 1000, ReapplyCondition="ReapplyDensity"))
 	float DensityAdjustmentFactor;
 
 	/** The minimum distance between foliage instances */
-	UPROPERTY(EditAnywhere, Category=Painting, meta=(UIMin=0, ClampMin=0, UIMax = 10000, ClampMax = 10000, ReapplyCondition="ReapplyRadius"))
+	UPROPERTY(EditAnywhere, Category=Painting, meta=(UIMin=0, ClampMin=0, UIMax = 1000, ClampMax = 1000, ReapplyCondition="ReapplyRadius"))
 	float Radius;
 
 	/** Specifies foliage instance scaling behavior when painting. */
diff --git a/Engine/Source/Runtime/Landscape/Classes/LandscapeGrassType.h b/Engine/Source/Runtime/Landscape/Classes/LandscapeGrassType.h
--- a/Engine/Source/Runtime/Landscape/Classes/LandscapeGrassType.h
+++ b/Engine/Source/Runtime/Landscape/Classes/LandscapeGrassType.h
@@ -32,7 +32,7 @@ struct FGrassVariety
 	UStaticMesh* GrassMesh;
 
 	/* Instances per 10 square meters. */
-	UPROPERTY(EditAnywhere, Category=Grass, meta = (UIMin = 0, ClampMin = 0, UIMax = 10000, ClampMax = 10000))
+	UPROPERTY(EditAnywhere, Category=Grass, meta = (UIMin = 0, ClampMin = 0, UIMax = 1000, ClampMax = 1000))
 	FPerPlatformFloat GrassDensity;
 
 	/* If true, use a jittered grid sequence for placement, otherwise use a halton sequence. */
