Fixed Instance count when Alt+Dragging instance

diff --git a/Engine/Source/Editor/FoliageEdit/Private/FoliageEdMode.cpp b/Engine/Source/Editor/FoliageEdit/Private/FoliageEdMode.cpp
--- a/Engine/Source/Editor/FoliageEdit/Private/FoliageEdMode.cpp
+++ b/Engine/Source/Editor/FoliageEdit/Private/FoliageEdMode.cpp
@@ -1818,6 +1818,7 @@ void FEdModeFoliage::TransformSelectedInstances(UWorld* InWorld, const FVector&
 					if (bDuplicate)
 					{
 						MeshInfo.DuplicateInstances(IFA, MeshPair.Key, SelectedIndices);
+						OnInstanceCountUpdated(MeshPair.Key);
 					}
 
 					MeshInfo.PreMoveInstances(IFA, SelectedIndices);
diff --git a/Engine/Source/Runtime/Engine/Public/StaticMeshResources.h b/Engine/Source/Runtime/Engine/Public/StaticMeshResources.h
--- a/Engine/Source/Runtime/Engine/Public/StaticMeshResources.h
+++ b/Engine/Source/Runtime/Engine/Public/StaticMeshResources.h
@@ -799,8 +799,7 @@ public:
 
 	/**
 	 * Constructor
-	 * @param InNeedsCPUAccess - true if resource array data should be CPU accessible
-	 * @param bSupportsVertexHalfFloat - true if device has support for half float in vertex arrays
+	 * @param bInUseHalfFloat - true if device has support for half float in vertex arrays
 	 */
 	FStaticMeshInstanceData(bool bInUseHalfFloat)
 		: bUseHalfFloat(PLATFORM_BUILTIN_VERTEX_HALF_FLOAT || bInUseHalfFloat)
diff --git a/Engine/Source/Runtime/Foliage/Private/InstancedFoliage.cpp b/Engine/Source/Runtime/Foliage/Private/InstancedFoliage.cpp
--- a/Engine/Source/Runtime/Foliage/Private/InstancedFoliage.cpp
+++ b/Engine/Source/Runtime/Foliage/Private/InstancedFoliage.cpp
@@ -1216,7 +1216,7 @@ void FFoliageMeshInfo::PreMoveInstances(AInstancedFoliageActor* InIFA, const TAr
 }
 
 
-void FFoliageMeshInfo::PostUpdateInstances(AInstancedFoliageActor* InIFA, const TArray<int32>& InInstancesUpdated, bool bReAddToHash)
+void FFoliageMeshInfo::PostUpdateInstances(AInstancedFoliageActor* InIFA, const TArray<int32>& InInstancesUpdated, bool bReAddToHash, bool InUpdateSelection)
 {
 	if (InInstancesUpdated.Num())
 	{
@@ -1236,6 +1236,12 @@ void FFoliageMeshInfo::PostUpdateInstances(AInstancedFoliageActor* InIFA, const
 			{
 				InstanceHash->InsertInstance(Instance.Location, InstanceIndex);
 			}
+
+			// Reselect the instance to update the render update to include selection as by default it gets removed
+			if (InUpdateSelection)
+			{
+				Component->SelectInstance(true, InstanceIndex);
+			}
 		}
 
 		Component->InvalidateLightingCache();
@@ -1245,7 +1251,7 @@ void FFoliageMeshInfo::PostUpdateInstances(AInstancedFoliageActor* InIFA, const
 
 void FFoliageMeshInfo::PostMoveInstances(AInstancedFoliageActor* InIFA, const TArray<int32>& InInstancesMoved)
 {
-	PostUpdateInstances(InIFA, InInstancesMoved, true);
+	PostUpdateInstances(InIFA, InInstancesMoved, true, true);
 }
 
 void FFoliageMeshInfo::DuplicateInstances(AInstancedFoliageActor* InIFA, UFoliageType* InSettings, const TArray<int32>& InInstancesToDuplicate)
diff --git a/Engine/Source/Runtime/Foliage/Public/InstancedFoliage.h b/Engine/Source/Runtime/Foliage/Public/InstancedFoliage.h
--- a/Engine/Source/Runtime/Foliage/Public/InstancedFoliage.h
+++ b/Engine/Source/Runtime/Foliage/Public/InstancedFoliage.h
@@ -207,7 +207,7 @@ struct FFoliageMeshInfo
 	FOLIAGE_API void CheckComponentClass(AInstancedFoliageActor* InIFA, const UFoliageType* InSettings);
 	FOLIAGE_API void PreMoveInstances(AInstancedFoliageActor* InIFA, const TArray<int32>& InInstancesToMove);
 	FOLIAGE_API void PostMoveInstances(AInstancedFoliageActor* InIFA, const TArray<int32>& InInstancesMoved);
-	FOLIAGE_API void PostUpdateInstances(AInstancedFoliageActor* InIFA, const TArray<int32>& InInstancesUpdated, bool bReAddToHash = false);
+	FOLIAGE_API void PostUpdateInstances(AInstancedFoliageActor* InIFA, const TArray<int32>& InInstancesUpdated, bool bReAddToHash = false, bool InUpdateSelection = false);
 	FOLIAGE_API void DuplicateInstances(AInstancedFoliageActor* InIFA, UFoliageType* InSettings, const TArray<int32>& InInstancesToDuplicate);
 	FOLIAGE_API void GetInstancesInsideSphere(const FSphere& Sphere, TArray<int32>& OutInstances);
 	FOLIAGE_API void GetInstanceAtLocation(const FVector& Location, int32& OutInstance, bool& bOutSucess);
