When doing undo, make the build foliage sync instead of async, as spamming the undo/redo would cause some issue in the handling of the async task result.

diff --git a/Engine/Source/Runtime/Foliage/Private/InstancedFoliage.cpp b/Engine/Source/Runtime/Foliage/Private/InstancedFoliage.cpp
--- a/Engine/Source/Runtime/Foliage/Private/InstancedFoliage.cpp
+++ b/Engine/Source/Runtime/Foliage/Private/InstancedFoliage.cpp
@@ -1372,7 +1372,7 @@ void FFoliageMeshInfo::ReapplyInstancesToComponent()
 		}
 
 		Component->bAutoRebuildTreeOnInstanceChanges = true;
-		Component->BuildTreeIfOutdated(true, true);
+		Component->BuildTreeIfOutdated(false, true);
 
 		Component->ClearInstanceSelection();
 
