Fixed ProceduralFoliageVolume Foliage Type removal not taken into account at next Resimulate

diff --git a/Engine/Source/Runtime/Foliage/Private/FoliageTypeObject.cpp b/Engine/Source/Runtime/Foliage/Private/FoliageTypeObject.cpp
--- a/Engine/Source/Runtime/Foliage/Private/FoliageTypeObject.cpp
+++ b/Engine/Source/Runtime/Foliage/Private/FoliageTypeObject.cpp
@@ -6,6 +6,7 @@
 void FFoliageTypeObject::RefreshInstance()
 {
 	//@todo: should verify that the instance is dirty before updating
+	TypeInstance = nullptr;
 
 	// If the foliage type is an asset, use that
 	if (auto FoliageTypeAsAsset = Cast<UFoliageType_InstancedStaticMesh>(FoliageTypeObject))
