Renaming this file before next submit to facilitate revision history

#rnx

diff --git a/Engine/Source/Runtime/Foliage/Private/InstancedFoliage.cpp b/Engine/Source/Runtime/Foliage/Private/InstancedFoliage.cpp
deleted file mode 100644
diff --git a/Engine/Source/Runtime/Foliage/Private/InstancedFoliageActor.cpp b/Engine/Source/Runtime/Foliage/Private/InstancedFoliageActor.cpp
new file mode 100644
--- /dev/null
+++ b/Engine/Source/Runtime/Foliage/Private/InstancedFoliageActor.cpp
@@ -0,0 +1,3453 @@
+// Copyright 1998-2019 Epic Games, Inc. All Rights Reserved.
+
+/*=============================================================================
+InstancedFoliage.cpp: Instanced foliage implementation.
+=============================================================================*/
+
+#include "InstancedFoliage.h"
+#include "Templates/SubclassOf.h"
+#include "HAL/IConsoleManager.h"
+#include "GameFramework/DamageType.h"
+#include "Engine/EngineTypes.h"
+#include "Components/SceneComponent.h"
+#include "Engine/Level.h"
+#include "GameFramework/Actor.h"
+#include "CollisionQueryParams.h"
+#include "WorldCollision.h"
+#include "Engine/Blueprint.h"
+#include "Engine/World.h"
+#include "Components/PrimitiveComponent.h"
+#include "FoliageType.h"
+#include "UObject/UObjectIterator.h"
+#include "FoliageInstancedStaticMeshComponent.h"
+#include "FoliageType_InstancedStaticMesh.h"
+#include "InstancedFoliageActor.h"
+#include "Serialization/CustomVersion.h"
+#include "UObject/Package.h"
+#include "UObject/PropertyPortFlags.h"
+#include "Engine/CollisionProfile.h"
+#include "Engine/Brush.h"
+#include "Engine/Engine.h"
+#include "Components/BrushComponent.h"
+#include "Components/ModelComponent.h"
+#include "Logging/TokenizedMessage.h"
+#include "Logging/MessageLog.h"
+#include "Misc/UObjectToken.h"
+#include "Misc/MapErrors.h"
+#include "ProceduralFoliageComponent.h"
+#include "ProceduralFoliageBlockingVolume.h"
+#include "ProceduralFoliageVolume.h"
+#include "EngineUtils.h"
+#include "EngineGlobals.h"
+#include "Engine/StaticMesh.h"
+#include "DrawDebugHelpers.h"
+#include "UObject/FortniteMainBranchObjectVersion.h"
+
+#define LOCTEXT_NAMESPACE "InstancedFoliage"
+
+#define DO_FOLIAGE_CHECK			0			// whether to validate foliage data during editing.
+#define FOLIAGE_CHECK_TRANSFORM		0			// whether to compare transforms between render and painting data.
+
+DEFINE_LOG_CATEGORY(LogInstancedFoliage);
+
+DECLARE_CYCLE_STAT(TEXT("FoliageActor_Trace"), STAT_FoliageTrace, STATGROUP_Foliage);
+DECLARE_CYCLE_STAT(TEXT("FoliageMeshInfo_AddInstance"), STAT_FoliageAddInstance, STATGROUP_Foliage);
+DECLARE_CYCLE_STAT(TEXT("FoliageMeshInfo_RemoveInstance"), STAT_FoliageRemoveInstance, STATGROUP_Foliage);
+DECLARE_CYCLE_STAT(TEXT("FoliageMeshInfo_CreateComponent"), STAT_FoliageCreateComponent, STATGROUP_Foliage);
+
+
+static TAutoConsoleVariable<int32> CVarFoliageDiscardDataOnLoad(
+	TEXT("foliage.DiscardDataOnLoad"),
+	0,
+	TEXT("1: Discard scalable foliage data on load (disables all scalable foliage types); 0: Keep scalable foliage data (requires reloading level)"),
+	ECVF_Scalability);
+
+// Custom serialization version for all packages containing Instance Foliage
+struct FFoliageCustomVersion
+{
+	enum Type
+	{
+		// Before any version changes were made in the plugin
+		BeforeCustomVersionWasAdded = 0,
+		// Converted to use HierarchicalInstancedStaticMeshComponent
+		FoliageUsingHierarchicalISMC = 1,
+		// Changed Component to not RF_Transactional
+		HierarchicalISMCNonTransactional = 2,
+		// Added FoliageTypeUpdateGuid
+		AddedFoliageTypeUpdateGuid = 3,
+		// Use a GUID to determine whic procedural actor spawned us
+		ProceduralGuid = 4,
+		// Support for cross-level bases 
+		CrossLevelBase = 5,
+		// FoliageType for details customization
+		FoliageTypeCustomization = 6,
+		// FoliageType for details customization continued
+		FoliageTypeCustomizationScaling = 7,
+		// FoliageType procedural scale and shade settings updated
+		FoliageTypeProceduralScaleAndShade = 8,
+		// Added FoliageHISMC and blueprint support
+		FoliageHISMCBlueprints = 9,
+		// Added Mobility setting to UFoliageType
+		AddedMobility = 10,
+		// Make sure that foliage has FoliageHISMC class
+		FoliageUsingFoliageISMC = 11,
+		// -----<new versions can be added above this line>-------------------------------------------------
+		VersionPlusOne,
+		LatestVersion = VersionPlusOne - 1
+	};
+
+	// The GUID for this custom version number
+	const static FGuid GUID;
+
+private:
+	FFoliageCustomVersion() {}
+};
+
+const FGuid FFoliageCustomVersion::GUID(0x430C4D19, 0x71544970, 0x87699B69, 0xDF90B0E5);
+// Register the custom version with core
+FCustomVersionRegistration GRegisterFoliageCustomVersion(FFoliageCustomVersion::GUID, FFoliageCustomVersion::LatestVersion, TEXT("FoliageVer"));
+
+
+// Legacy (< FFoliageCustomVersion::CrossLevelBase) serializer
+FArchive& operator<<(FArchive& Ar, FFoliageInstance_Deprecated& Instance)
+{
+	Ar << Instance.Base;
+	Ar << Instance.Location;
+	Ar << Instance.Rotation;
+	Ar << Instance.DrawScale3D;
+
+	if (Ar.CustomVer(FFoliageCustomVersion::GUID) < FFoliageCustomVersion::FoliageUsingHierarchicalISMC)
+	{
+		int32 OldClusterIndex;
+		Ar << OldClusterIndex;
+		Ar << Instance.PreAlignRotation;
+		Ar << Instance.Flags;
+
+		if (OldClusterIndex == INDEX_NONE)
+		{
+			// When converting, we need to skip over any instance that was previously deleted but still in the Instances array.
+			Instance.Flags |= FOLIAGE_InstanceDeleted;
+		}
+	}
+	else
+	{
+		Ar << Instance.PreAlignRotation;
+		Ar << Instance.Flags;
+	}
+
+	Ar << Instance.ZOffset;
+
+#if WITH_EDITORONLY_DATA
+	if (!Ar.ArIsFilterEditorOnly && Ar.CustomVer(FFoliageCustomVersion::GUID) >= FFoliageCustomVersion::ProceduralGuid)
+	{
+		Ar << Instance.ProceduralGuid;
+	}
+#endif
+
+	return Ar;
+}
+
+//
+// Serializers for struct data
+//
+FArchive& operator<<(FArchive& Ar, FFoliageInstance& Instance)
+{
+	Ar << Instance.Location;
+	Ar << Instance.Rotation;
+	Ar << Instance.DrawScale3D;
+	Ar << Instance.PreAlignRotation;
+	Ar << Instance.ProceduralGuid;
+	Ar << Instance.Flags;
+	Ar << Instance.ZOffset;
+	Ar << Instance.BaseId;
+
+	return Ar;
+}
+
+static void ConvertDeprecatedFoliageMeshes(
+	AInstancedFoliageActor* IFA,
+	const TMap<UFoliageType*, TUniqueObj<FFoliageMeshInfo_Deprecated>>& FoliageMeshesDeprecated,
+	TMap<UFoliageType*, TUniqueObj<FFoliageMeshInfo>>& FoliageMeshes)
+{
+#if WITH_EDITORONLY_DATA	
+	for (auto Pair : FoliageMeshesDeprecated)
+	{
+		auto& FoliageMesh = FoliageMeshes.Add(Pair.Key);
+		const auto& FoliageMeshDeprecated = Pair.Value;
+
+		FoliageMesh->Component = FoliageMeshDeprecated->Component;
+		FoliageMesh->FoliageTypeUpdateGuid = FoliageMeshDeprecated->FoliageTypeUpdateGuid;
+
+		FoliageMesh->Instances.Reserve(FoliageMeshDeprecated->Instances.Num());
+
+		for (const FFoliageInstance_Deprecated& DeprecatedInstance : FoliageMeshDeprecated->Instances)
+		{
+			FFoliageInstance Instance;
+			static_cast<FFoliageInstancePlacementInfo&>(Instance) = DeprecatedInstance;
+			Instance.BaseId = IFA->InstanceBaseCache.AddInstanceBaseId(DeprecatedInstance.Base);
+			Instance.ProceduralGuid = DeprecatedInstance.ProceduralGuid;
+
+			FoliageMesh->Instances.Add(Instance);
+		}
+	}
+
+	// there were no cross-level references before
+	check(IFA->InstanceBaseCache.InstanceBaseLevelMap.Num() <= 1);
+	// populate WorldAsset->BasePtr map
+	IFA->InstanceBaseCache.InstanceBaseLevelMap.Empty();
+	auto& BaseList = IFA->InstanceBaseCache.InstanceBaseLevelMap.Add(TSoftObjectPtr<UWorld>(Cast<UWorld>(IFA->GetLevel()->GetOuter())));
+	for (auto& BaseInfoPair : IFA->InstanceBaseCache.InstanceBaseMap)
+	{
+		BaseList.Add(BaseInfoPair.Value.BasePtr);
+	}
+#endif//WITH_EDITORONLY_DATA	
+}
+
+/**
+*	FFoliageInstanceCluster_Deprecated
+*/
+struct FFoliageInstanceCluster_Deprecated
+{
+	UInstancedStaticMeshComponent* ClusterComponent;
+	FBoxSphereBounds Bounds;
+
+#if WITH_EDITORONLY_DATA
+	TArray<int32> InstanceIndices;	// index into editor editor Instances array
+#endif
+
+	friend FArchive& operator<<(FArchive& Ar, FFoliageInstanceCluster_Deprecated& OldCluster)
+	{
+		check(Ar.CustomVer(FFoliageCustomVersion::GUID) < FFoliageCustomVersion::FoliageUsingHierarchicalISMC);
+
+		Ar << OldCluster.Bounds;
+		Ar << OldCluster.ClusterComponent;
+
+#if WITH_EDITORONLY_DATA
+		if (!Ar.ArIsFilterEditorOnly ||
+			Ar.UE4Ver() < VER_UE4_FOLIAGE_SETTINGS_TYPE)
+		{
+			Ar << OldCluster.InstanceIndices;
+		}
+#endif
+
+		return Ar;
+	}
+};
+
+FArchive& operator<<(FArchive& Ar, FFoliageMeshInfo_Deprecated& MeshInfo)
+{
+	if (Ar.CustomVer(FFoliageCustomVersion::GUID) >= FFoliageCustomVersion::FoliageUsingHierarchicalISMC)
+	{
+		Ar << MeshInfo.Component;
+	}
+	else
+	{
+		TArray<FFoliageInstanceCluster_Deprecated> OldInstanceClusters;
+		Ar << OldInstanceClusters;
+	}
+
+#if WITH_EDITORONLY_DATA
+	if ((!Ar.ArIsFilterEditorOnly || Ar.UE4Ver() < VER_UE4_FOLIAGE_SETTINGS_TYPE) &&
+		(!(Ar.GetPortFlags() & PPF_DuplicateForPIE)))
+	{
+		Ar << MeshInfo.Instances;
+	}
+
+	if (!Ar.ArIsFilterEditorOnly && Ar.CustomVer(FFoliageCustomVersion::GUID) >= FFoliageCustomVersion::AddedFoliageTypeUpdateGuid)
+	{
+		Ar << MeshInfo.FoliageTypeUpdateGuid;
+	}
+#endif
+
+	return Ar;
+}
+
+FArchive& operator<<(FArchive& Ar, FFoliageMeshInfo& MeshInfo)
+{
+	Ar << MeshInfo.Component;
+
+#if WITH_EDITORONLY_DATA
+	if (!Ar.ArIsFilterEditorOnly && !(Ar.GetPortFlags() & PPF_DuplicateForPIE))
+	{
+		if (Ar.IsTransacting())
+		{
+			MeshInfo.Instances.BulkSerialize(Ar);
+		}
+		else
+		{
+			Ar << MeshInfo.Instances;
+		}
+	}
+
+	if (!Ar.ArIsFilterEditorOnly)
+	{
+		Ar << MeshInfo.FoliageTypeUpdateGuid;
+	}
+
+	// Serialize the transient data for undo.
+	if (Ar.IsTransacting())
+	{
+		Ar << MeshInfo.ComponentHash;
+		Ar << MeshInfo.SelectedIndices;
+	}
+#endif
+
+	return Ar;
+}
+
+//
+// UFoliageType
+//
+
+UFoliageType::UFoliageType(const FObjectInitializer& ObjectInitializer)
+	: Super(ObjectInitializer)
+{
+	Density = 100.0f;
+	Radius = 0.0f;
+	AlignToNormal = true;
+	RandomYaw = true;
+	Scaling = EFoliageScaling::Uniform;
+	ScaleX.Min = 1.0f;
+	ScaleY.Min = 1.0f;
+	ScaleZ.Min = 1.0f;
+	ScaleX.Max = 1.0f;
+	ScaleY.Max = 1.0f;
+	ScaleZ.Max = 1.0f;
+	AlignMaxAngle = 0.0f;
+	RandomPitchAngle = 0.0f;
+	GroundSlopeAngle.Min = 0.0f;
+	GroundSlopeAngle.Max = 45.0f;
+	Height.Min = -262144.0f;
+	Height.Max = 262144.0f;
+	ZOffset.Min = 0.0f;
+	ZOffset.Max = 0.0f;
+	CullDistance.Min = 0;
+	CullDistance.Max = 0;
+	bEnableStaticLighting_DEPRECATED = true;
+	MinimumLayerWeight = 0.5f;
+#if WITH_EDITORONLY_DATA
+	IsSelected = false;
+#endif
+	DensityAdjustmentFactor = 1.0f;
+	CollisionWithWorld = false;
+	CollisionScale = FVector(0.9f, 0.9f, 0.9f);
+
+	Mobility = EComponentMobility::Static;
+	CastShadow = true;
+	bCastDynamicShadow = true;
+	bCastStaticShadow = true;
+	bAffectDynamicIndirectLighting = false;
+	// Most of the high instance count foliage like grass causes performance problems with distance field lighting
+	bAffectDistanceFieldLighting = false;
+	bCastShadowAsTwoSided = false;
+	bReceivesDecals = false;
+
+	bOverrideLightMapRes = false;
+	OverriddenLightMapRes = 8;
+	bUseAsOccluder = false;
+
+	BodyInstance.SetCollisionProfileName(UCollisionProfile::NoCollision_ProfileName);
+
+	/** Ecosystem settings*/
+	AverageSpreadDistance = 50;
+	SpreadVariance = 150;
+	bCanGrowInShade = false;
+	bSpawnsInShade = false;
+	SeedsPerStep = 3;
+	OverlapPriority = 0.f;
+	NumSteps = 3;
+	ProceduralScale = FFloatInterval(1.f, 3.f);
+	ChangeCount = 0;
+	InitialSeedDensity = 1.f;
+	CollisionRadius = 100.f;
+	ShadeRadius = 100.f;
+	MaxInitialAge = 0.f;
+	MaxAge = 10.f;
+
+	FRichCurve* Curve = ScaleCurve.GetRichCurve();
+	Curve->AddKey(0.f, 0.f);
+	Curve->AddKey(1.f, 1.f);
+
+	UpdateGuid = FGuid::NewGuid();
+#if WITH_EDITORONLY_DATA
+	HiddenEditorViews = 0;
+#endif
+	bEnableDensityScaling = false;
+
+#if WITH_EDITORONLY_DATA
+	// Deprecated since FFoliageCustomVersion::FoliageTypeCustomization
+	ScaleMinX_DEPRECATED = 1.0f;
+	ScaleMinY_DEPRECATED = 1.0f;
+	ScaleMinZ_DEPRECATED = 1.0f;
+	ScaleMaxX_DEPRECATED = 1.0f;
+	ScaleMaxY_DEPRECATED = 1.0f;
+	ScaleMaxZ_DEPRECATED = 1.0f;
+	HeightMin_DEPRECATED = -262144.0f;
+	HeightMax_DEPRECATED = 262144.0f;
+	ZOffsetMin_DEPRECATED = 0.0f;
+	ZOffsetMax_DEPRECATED = 0.0f;
+	UniformScale_DEPRECATED = true;
+	GroundSlope_DEPRECATED = 45.0f;
+
+	// Deprecated since FFoliageCustomVersion::FoliageTypeProceduralScaleAndShade
+	MinScale_DEPRECATED = 1.f;
+	MaxScale_DEPRECATED = 3.f;
+
+#endif// WITH_EDITORONLY_DATA
+}
+
+void UFoliageType::Serialize(FArchive& Ar)
+{
+	Super::Serialize(Ar);
+
+	Ar.UsingCustomVersion(FFoliageCustomVersion::GUID);
+
+	// we now have mask configurations for every color channel
+	if (Ar.IsLoading() && Ar.IsPersistent() && !Ar.HasAnyPortFlags(PPF_Duplicate | PPF_DuplicateForPIE) && VertexColorMask_DEPRECATED != FOLIAGEVERTEXCOLORMASK_Disabled)
+	{
+		FFoliageVertexColorChannelMask* Mask = nullptr;
+		switch (VertexColorMask_DEPRECATED)
+		{
+		case FOLIAGEVERTEXCOLORMASK_Red:
+			Mask = &VertexColorMaskByChannel[(uint8)EVertexColorMaskChannel::Red];
+			break;
+
+		case FOLIAGEVERTEXCOLORMASK_Green:
+			Mask = &VertexColorMaskByChannel[(uint8)EVertexColorMaskChannel::Green];
+			break;
+
+		case FOLIAGEVERTEXCOLORMASK_Blue:
+			Mask = &VertexColorMaskByChannel[(uint8)EVertexColorMaskChannel::Blue];
+			break;
+
+		case FOLIAGEVERTEXCOLORMASK_Alpha:
+			Mask = &VertexColorMaskByChannel[(uint8)EVertexColorMaskChannel::Alpha];
+			break;
+		}
+
+		if (Mask != nullptr)
+		{
+			Mask->UseMask = true;
+			Mask->MaskThreshold = VertexColorMaskThreshold_DEPRECATED;
+			Mask->InvertMask = VertexColorMaskInvert_DEPRECATED;
+
+			VertexColorMask_DEPRECATED = FOLIAGEVERTEXCOLORMASK_Disabled;
+		}
+	}
+
+	if (LandscapeLayer_DEPRECATED != NAME_None && LandscapeLayers.Num() == 0)	//we now store an array of names so initialize the array with the old name
+	{
+		LandscapeLayers.Add(LandscapeLayer_DEPRECATED);
+		LandscapeLayer_DEPRECATED = NAME_None;
+	}
+
+	if (Ar.IsLoading() && GetLinkerCustomVersion(FFoliageCustomVersion::GUID) < FFoliageCustomVersion::AddedMobility)
+	{
+		Mobility = bEnableStaticLighting_DEPRECATED ? EComponentMobility::Static : EComponentMobility::Movable;
+	}
+
+#if WITH_EDITORONLY_DATA
+	if (Ar.IsLoading())
+	{
+		if (Ar.CustomVer(FFoliageCustomVersion::GUID) < FFoliageCustomVersion::FoliageTypeCustomization)
+		{
+			ScaleX.Min = ScaleMinX_DEPRECATED;
+			ScaleX.Max = ScaleMaxX_DEPRECATED;
+
+			ScaleY.Min = ScaleMinY_DEPRECATED;
+			ScaleY.Max = ScaleMaxY_DEPRECATED;
+
+			ScaleZ.Min = ScaleMinZ_DEPRECATED;
+			ScaleZ.Max = ScaleMaxZ_DEPRECATED;
+
+			Height.Min = HeightMin_DEPRECATED;
+			Height.Max = HeightMax_DEPRECATED;
+
+			ZOffset.Min = ZOffsetMin_DEPRECATED;
+			ZOffset.Max = ZOffsetMax_DEPRECATED;
+
+			CullDistance.Min = StartCullDistance_DEPRECATED;
+			CullDistance.Max = EndCullDistance_DEPRECATED;
+		}
+
+		if (Ar.CustomVer(FFoliageCustomVersion::GUID) < FFoliageCustomVersion::FoliageTypeCustomizationScaling)
+		{
+			Scaling = UniformScale_DEPRECATED ? EFoliageScaling::Uniform : EFoliageScaling::Free;
+
+			GroundSlopeAngle.Min = MinGroundSlope_DEPRECATED;
+			GroundSlopeAngle.Max = GroundSlope_DEPRECATED;
+		}
+
+		if (Ar.CustomVer(FFoliageCustomVersion::GUID) < FFoliageCustomVersion::FoliageTypeProceduralScaleAndShade)
+		{
+			bCanGrowInShade = bSpawnsInShade;
+
+			ProceduralScale.Min = MinScale_DEPRECATED;
+			ProceduralScale.Max = MaxScale_DEPRECATED;
+		}
+	}
+#endif// WITH_EDITORONLY_DATA
+}
+
+bool UFoliageType::IsNotAssetOrBlueprint() const
+{
+	return IsAsset() == false && Cast<UBlueprint>(GetClass()->ClassGeneratedBy) == nullptr;
+}
+
+FVector UFoliageType::GetRandomScale() const
+{
+	FVector Result(1.0f);
+	float LockRand = 0.0f;
+
+	switch (Scaling)
+	{
+	case EFoliageScaling::Uniform:
+		Result.X = ScaleX.Interpolate(FMath::FRand());
+		Result.Y = Result.X;
+		Result.Z = Result.X;
+		break;
+
+	case EFoliageScaling::Free:
+		Result.X = ScaleX.Interpolate(FMath::FRand());
+		Result.Y = ScaleY.Interpolate(FMath::FRand());
+		Result.Z = ScaleZ.Interpolate(FMath::FRand());
+		break;
+
+	case EFoliageScaling::LockXY:
+		LockRand = FMath::FRand();
+		Result.X = ScaleX.Interpolate(LockRand);
+		Result.Y = ScaleY.Interpolate(LockRand);
+		Result.Z = ScaleZ.Interpolate(FMath::FRand());
+		break;
+
+	case EFoliageScaling::LockXZ:
+		LockRand = FMath::FRand();
+		Result.X = ScaleX.Interpolate(LockRand);
+		Result.Y = ScaleY.Interpolate(FMath::FRand());
+		Result.Z = ScaleZ.Interpolate(LockRand);
+
+	case EFoliageScaling::LockYZ:
+		LockRand = FMath::FRand();
+		Result.X = ScaleX.Interpolate(FMath::FRand());
+		Result.Y = ScaleY.Interpolate(LockRand);
+		Result.Z = ScaleZ.Interpolate(LockRand);
+	}
+
+	return Result;
+}
+
+UFoliageType_InstancedStaticMesh::UFoliageType_InstancedStaticMesh(const FObjectInitializer& ObjectInitializer)
+	: Super(ObjectInitializer)
+{
+	Mesh = nullptr;
+	ComponentClass = UFoliageInstancedStaticMeshComponent::StaticClass();
+	CustomNavigableGeometry = EHasCustomNavigableGeometry::Yes;
+}
+
+
+float UFoliageType::GetMaxRadius() const
+{
+	return FMath::Max(CollisionRadius, ShadeRadius);
+}
+
+float UFoliageType::GetScaleForAge(const float Age) const
+{
+	const FRichCurve* Curve = ScaleCurve.GetRichCurveConst();
+	const float Time = FMath::Clamp(MaxAge == 0 ? 1.f : Age / MaxAge, 0.f, 1.f);
+	const float Scale = Curve->Eval(Time);
+	return ProceduralScale.Min + ProceduralScale.Size() * Scale;
+}
+
+float UFoliageType::GetInitAge(FRandomStream& RandomStream) const
+{
+	return RandomStream.FRandRange(0, MaxInitialAge);
+}
+
+float UFoliageType::GetNextAge(const float CurrentAge, const int32 InNumSteps) const
+{
+	float NewAge = CurrentAge;
+	for (int32 Count = 0; Count < InNumSteps; ++Count)
+	{
+		const float GrowAge = NewAge + 1;
+		if (GrowAge <= MaxAge)
+		{
+			NewAge = GrowAge;
+		}
+		else
+		{
+			break;
+		}
+	}
+
+	return NewAge;
+}
+
+bool UFoliageType::GetSpawnsInShade() const
+{
+	return bCanGrowInShade && bSpawnsInShade;
+}
+
+#if WITH_EDITOR
+void UFoliageType::PostEditChangeProperty(struct FPropertyChangedEvent& PropertyChangedEvent)
+{
+	Super::PostEditChangeProperty(PropertyChangedEvent);
+
+	// Ensure that OverriddenLightMapRes is a factor of 4
+	OverriddenLightMapRes = OverriddenLightMapRes > 4 ? OverriddenLightMapRes + 3 & ~3 : 4;
+	++ChangeCount;
+
+	UpdateGuid = FGuid::NewGuid();
+
+	//@todo: move this into FoliageType_InstancedStaticMesh
+	// Check to see if the mesh is what changed
+	const bool bMeshChanged = PropertyChangedEvent.Property && PropertyChangedEvent.Property->GetFName() == GET_MEMBER_NAME_CHECKED(UFoliageType_InstancedStaticMesh, Mesh);
+
+	// Notify any currently-loaded InstancedFoliageActors
+	if (IsFoliageReallocationRequiredForPropertyChange(PropertyChangedEvent))
+	{
+		for (TObjectIterator<AInstancedFoliageActor> It(RF_ClassDefaultObject, /** bIncludeDerivedClasses */ true, /** InternalExcludeFalgs */ EInternalObjectFlags::PendingKill); It; ++It)
+		{
+			if (It->GetWorld() != nullptr)
+			{
+				It->NotifyFoliageTypeChanged(this, bMeshChanged);
+			}
+		}
+	}
+}
+
+void UFoliageType::PreEditChange(UProperty* PropertyAboutToChange)
+{
+	Super::PreEditChange(PropertyAboutToChange);
+
+	if (PropertyAboutToChange && PropertyAboutToChange->GetFName() == GET_MEMBER_NAME_CHECKED(UFoliageType_InstancedStaticMesh, Mesh))
+	{
+		for (TObjectIterator<AInstancedFoliageActor> It(RF_ClassDefaultObject, /** bIncludeDerivedClasses */ true, /** InternalExcludeFalgs */ EInternalObjectFlags::PendingKill); It; ++It)
+		{
+			It->NotifyFoliageTypeWillChange(this, true);
+		}
+	}
+}
+
+void UFoliageType::OnHiddenEditorViewMaskChanged(UWorld* InWorld)
+{
+	for (TActorIterator<AInstancedFoliageActor> It(InWorld); It; ++It)
+	{
+		FFoliageMeshInfo* MeshInfo = It->FindMesh(this);
+		if (MeshInfo && MeshInfo->Component)
+		{
+			UFoliageInstancedStaticMeshComponent* FoliageComponent = Cast<UFoliageInstancedStaticMeshComponent>(MeshInfo->Component);
+
+			if (FoliageComponent && FoliageComponent->FoliageHiddenEditorViews != HiddenEditorViews)
+			{
+				FoliageComponent->FoliageHiddenEditorViews = HiddenEditorViews;
+				FoliageComponent->MarkRenderStateDirty();
+			}
+		}
+	}
+}
+
+FName UFoliageType::GetDisplayFName() const
+{
+	FName DisplayFName;
+
+	if (IsAsset())
+	{
+		DisplayFName = GetFName();
+	}
+	else if (UBlueprint* FoliageTypeBP = Cast<UBlueprint>(GetClass()->ClassGeneratedBy))
+	{
+		DisplayFName = FoliageTypeBP->GetFName();
+	}
+	else if (UStaticMesh* StaticMesh = GetStaticMesh())
+	{
+		DisplayFName = StaticMesh->GetFName();
+	}
+
+	return DisplayFName;
+}
+
+#endif
+
+//
+// FFoliageMeshInfo
+//
+FFoliageMeshInfo::FFoliageMeshInfo()
+	: Component(nullptr)
+#if WITH_EDITOR
+	, InstanceHash(GIsEditor ? new FFoliageInstanceHash() : nullptr)
+#endif
+{ }
+
+FFoliageMeshInfo::~FFoliageMeshInfo()
+{ }
+
+
+#if WITH_EDITOR
+
+void FFoliageMeshInfo::CheckValid()
+{
+#if DO_FOLIAGE_CHECK
+	int32 ClusterTotal = 0;
+	int32 ComponentTotal = 0;
+
+	for (FFoliageInstanceCluster& Cluster : InstanceClusters)
+	{
+		check(Cluster.ClusterComponent != nullptr);
+		ClusterTotal += Cluster.InstanceIndices.Num();
+		ComponentTotal += Cluster.ClusterComponent->PerInstanceSMData.Num();
+	}
+
+	check(ClusterTotal == ComponentTotal);
+
+	int32 FreeTotal = 0;
+	int32 InstanceTotal = 0;
+	for (int32 InstanceIdx = 0; InstanceIdx < Instances.Num(); InstanceIdx++)
+	{
+		if (Instances[InstanceIdx].ClusterIndex != -1)
+		{
+			InstanceTotal++;
+		}
+		else
+		{
+			FreeTotal++;
+		}
+	}
+
+	check(ClusterTotal == InstanceTotal);
+	check(FreeInstanceIndices.Num() == FreeTotal);
+
+	InstanceHash->CheckInstanceCount(InstanceTotal);
+
+	int32 ComponentHashTotal = 0;
+	for (const auto& Pair : ComponentHash)
+	{
+		ComponentHashTotal += Pair.Value().Num();
+	}
+	check(ComponentHashTotal == InstanceTotal);
+
+#if FOLIAGE_CHECK_TRANSFORM
+	// Check transforms match up with editor data
+	int32 MismatchCount = 0;
+	for (int32 ClusterIdx = 0; ClusterIdx<InstanceClusters.Num(); ClusterIdx++)
+	{
+		TArray<int32> Indices = InstanceClusters(ClusterIdx).InstanceIndices;
+		UInstancedStaticMeshComponent* Comp = InstanceClusters(ClusterIdx).ClusterComponent;
+		for (int32 InstIdx = 0; InstIdx<Indices.Num(); InstIdx++)
+		{
+			int32 InstanceIdx = Indices(InstIdx);
+
+			FTransform InstanceToWorldEd = Instances(InstanceIdx).GetInstanceTransform();
+			FTransform InstanceToWorldCluster = Comp->PerInstanceSMData(InstIdx).Transform * Comp->GetComponentToWorld();
+
+			if (!InstanceToWorldEd.Equals(InstanceToWorldCluster))
+			{
+				Comp->PerInstanceSMData(InstIdx).Transform = InstanceToWorldEd.ToMatrixWithScale();
+				MismatchCount++;
+			}
+		}
+	}
+
+	if (MismatchCount != 0)
+	{
+		UE_LOG(LogInstancedFoliage, Log, TEXT("%s: transform mismatch: %d"), *InstanceClusters(0).ClusterComponent->StaticMesh->GetName(), MismatchCount);
+	}
+#endif
+
+#endif
+}
+
+void FFoliageMeshInfo::CreateNewComponent(AInstancedFoliageActor* InIFA, const UFoliageType* InSettings)
+{
+	SCOPE_CYCLE_COUNTER(STAT_FoliageCreateComponent);
+	
+	check(!Component);
+
+	UClass* ComponentClass = InSettings->GetComponentClass();
+	if (ComponentClass == nullptr)
+	{
+		ComponentClass = UFoliageInstancedStaticMeshComponent::StaticClass();
+	}
+
+	UFoliageInstancedStaticMeshComponent* FoliageComponent = NewObject<UFoliageInstancedStaticMeshComponent>(InIFA, ComponentClass, NAME_None, RF_Transactional);
+
+	Component = FoliageComponent;
+	Component->SetStaticMesh(InSettings->GetStaticMesh());
+	Component->bSelectable = true;
+	Component->bHasPerInstanceHitProxies = true;
+
+	if (Component->GetStaticMesh() != nullptr)
+	{
+		Component->GetStaticMesh()->GetOnExtendedBoundsChanged().AddRaw(this, &FFoliageMeshInfo::HandleComponentMeshBoundsChanged);
+	}
+
+#if WITH_EDITOR
+	FoliageComponent->FoliageHiddenEditorViews = InSettings->HiddenEditorViews;
+#endif	
+
+	UpdateComponentSettings(InSettings);
+
+	Component->SetupAttachment(InIFA->GetRootComponent());
+
+	if (InIFA->GetRootComponent()->IsRegistered())
+	{
+		Component->RegisterComponent();
+	}
+
+	// Use only instance translation as a component transform
+	Component->SetWorldTransform(InIFA->GetRootComponent()->GetComponentTransform());
+
+	// Add the new component to the transaction buffer so it will get destroyed on undo
+	Component->Modify();
+	// We don't want to track changes to instances later so we mark it as non-transactional
+	Component->ClearFlags(RF_Transactional);
+}
+
+void FFoliageMeshInfo::HandleComponentMeshBoundsChanged(const FBoxSphereBounds& NewBounds)
+{
+	if (Component != nullptr)
+	{
+		Component->BuildTreeIfOutdated(true, false);
+	}
+}
+
+void FFoliageMeshInfo::CheckComponentClass(AInstancedFoliageActor* InIFA, const UFoliageType* InSettings)
+{
+	if (Component)
+	{
+		UClass* ComponentClass = InSettings->GetComponentClass();
+		if (ComponentClass == nullptr)
+		{
+			ComponentClass = UFoliageInstancedStaticMeshComponent::StaticClass();
+		}
+
+		if (ComponentClass != Component->GetClass())
+		{
+			InIFA->Modify();
+
+			// prepare to destroy the old component
+			Component->ClearInstances();
+
+			// make sure the destruction gets stored in the undo buffer, 
+			// so the component will get recreated on undo.
+			Component->SetFlags(RF_Transactional);
+			Component->Modify();
+
+			Component->DestroyComponent();
+			Component = nullptr;
+
+			// create a new component
+			CreateNewComponent(InIFA, InSettings);
+
+			// apply the instances to it
+			ReapplyInstancesToComponent();
+		}
+	}
+}
+
+void FFoliageMeshInfo::UpdateComponentSettings(const UFoliageType* InSettings)
+{
+	if (Component)
+	{
+		bool bNeedsMarkRenderStateDirty = false;
+		bool bNeedsInvalidateLightingCache = false;
+
+		const UFoliageType* FoliageType = InSettings;
+		if (InSettings->GetClass()->ClassGeneratedBy)
+		{
+			// If we're updating settings for a BP foliage type, use the CDO
+			FoliageType = InSettings->GetClass()->GetDefaultObject<UFoliageType>();
+		}
+
+		if (Component->GetStaticMesh() != FoliageType->GetStaticMesh())
+		{
+			Component->SetStaticMesh(FoliageType->GetStaticMesh());
+
+			bNeedsInvalidateLightingCache = true;
+			bNeedsMarkRenderStateDirty = true;
+		}
+
+		if (Component->Mobility != FoliageType->Mobility)
+		{
+			Component->SetMobility(FoliageType->Mobility);
+			bNeedsMarkRenderStateDirty = true;
+			bNeedsInvalidateLightingCache = true;
+		}
+		if (Component->InstanceStartCullDistance != FoliageType->CullDistance.Min)
+		{
+			Component->InstanceStartCullDistance = FoliageType->CullDistance.Min;
+			bNeedsMarkRenderStateDirty = true;
+		}
+		if (Component->InstanceEndCullDistance != FoliageType->CullDistance.Max)
+		{
+			Component->InstanceEndCullDistance = FoliageType->CullDistance.Max;
+			bNeedsMarkRenderStateDirty = true;
+		}
+		if (Component->CastShadow != FoliageType->CastShadow)
+		{
+			Component->CastShadow = FoliageType->CastShadow;
+			bNeedsMarkRenderStateDirty = true;
+			bNeedsInvalidateLightingCache = true;
+		}
+		if (Component->bCastDynamicShadow != FoliageType->bCastDynamicShadow)
+		{
+			Component->bCastDynamicShadow = FoliageType->bCastDynamicShadow;
+			bNeedsMarkRenderStateDirty = true;
+			bNeedsInvalidateLightingCache = true;
+		}
+		if (Component->bCastStaticShadow != FoliageType->bCastStaticShadow)
+		{
+			Component->bCastStaticShadow = FoliageType->bCastStaticShadow;
+			bNeedsMarkRenderStateDirty = true;
+			bNeedsInvalidateLightingCache = true;
+		}
+		if (Component->bAffectDynamicIndirectLighting != FoliageType->bAffectDynamicIndirectLighting)
+		{
+			Component->bAffectDynamicIndirectLighting = FoliageType->bAffectDynamicIndirectLighting;
+			bNeedsMarkRenderStateDirty = true;
+			bNeedsInvalidateLightingCache = true;
+		}
+		if (Component->bAffectDistanceFieldLighting != FoliageType->bAffectDistanceFieldLighting)
+		{
+			Component->bAffectDistanceFieldLighting = FoliageType->bAffectDistanceFieldLighting;
+			bNeedsMarkRenderStateDirty = true;
+			bNeedsInvalidateLightingCache = true;
+		}
+		if (Component->bCastShadowAsTwoSided != FoliageType->bCastShadowAsTwoSided)
+		{
+			Component->bCastShadowAsTwoSided = FoliageType->bCastShadowAsTwoSided;
+			bNeedsMarkRenderStateDirty = true;
+			bNeedsInvalidateLightingCache = true;
+		}
+		if (Component->bReceivesDecals != FoliageType->bReceivesDecals)
+		{
+			Component->bReceivesDecals = FoliageType->bReceivesDecals;
+			bNeedsMarkRenderStateDirty = true;
+			bNeedsInvalidateLightingCache = true;
+		}
+		if (Component->bOverrideLightMapRes != FoliageType->bOverrideLightMapRes)
+		{
+			Component->bOverrideLightMapRes = FoliageType->bOverrideLightMapRes;
+			bNeedsMarkRenderStateDirty = true;
+			bNeedsInvalidateLightingCache = true;
+		}
+		if (Component->OverriddenLightMapRes != FoliageType->OverriddenLightMapRes)
+		{
+			Component->OverriddenLightMapRes = FoliageType->OverriddenLightMapRes;
+			bNeedsMarkRenderStateDirty = true;
+			bNeedsInvalidateLightingCache = true;
+		}
+		if (Component->LightmapType != FoliageType->LightmapType)
+		{
+			Component->LightmapType = FoliageType->LightmapType;
+			bNeedsMarkRenderStateDirty = true;
+			bNeedsInvalidateLightingCache = true;
+		}
+		if (Component->bUseAsOccluder != FoliageType->bUseAsOccluder)
+		{
+			Component->bUseAsOccluder = FoliageType->bUseAsOccluder;
+			bNeedsMarkRenderStateDirty = true;
+		}
+
+		if (Component->bEnableDensityScaling != FoliageType->bEnableDensityScaling)
+		{
+			Component->bEnableDensityScaling = FoliageType->bEnableDensityScaling;
+
+			Component->UpdateDensityScaling();
+
+			bNeedsMarkRenderStateDirty = true;
+		}
+
+		if (GetLightingChannelMaskForStruct(Component->LightingChannels) != GetLightingChannelMaskForStruct(FoliageType->LightingChannels))
+		{
+			Component->LightingChannels = FoliageType->LightingChannels;
+			bNeedsMarkRenderStateDirty = true;
+		}
+
+		UFoliageInstancedStaticMeshComponent* FoliageComponent = Cast<UFoliageInstancedStaticMeshComponent>(Component);
+
+		if (FoliageComponent && FoliageComponent->FoliageHiddenEditorViews != InSettings->HiddenEditorViews)
+		{
+			FoliageComponent->FoliageHiddenEditorViews = InSettings->HiddenEditorViews;
+			bNeedsMarkRenderStateDirty = true;
+		}
+
+		if (Component->bRenderCustomDepth != FoliageType->bRenderCustomDepth)
+		{
+			Component->bRenderCustomDepth = FoliageType->bRenderCustomDepth;
+			bNeedsMarkRenderStateDirty = true;
+		}
+
+		if (Component->CustomDepthStencilValue != FoliageType->CustomDepthStencilValue)
+		{
+			Component->CustomDepthStencilValue = FoliageType->CustomDepthStencilValue;
+			bNeedsMarkRenderStateDirty = true;
+		}
+
+		const UFoliageType_InstancedStaticMesh* FoliageType_ISM = Cast<UFoliageType_InstancedStaticMesh>(FoliageType);
+		if (FoliageType_ISM)
+		{
+			// Check override materials
+			if (Component->OverrideMaterials.Num() != FoliageType_ISM->OverrideMaterials.Num())
+			{
+				Component->OverrideMaterials = FoliageType_ISM->OverrideMaterials;
+				bNeedsMarkRenderStateDirty = true;
+				bNeedsInvalidateLightingCache = true;
+			}
+			else
+			{
+				for (int32 Index = 0; Index < FoliageType_ISM->OverrideMaterials.Num(); Index++)
+				{
+					if (Component->OverrideMaterials[Index] != FoliageType_ISM->OverrideMaterials[Index])
+					{
+						Component->OverrideMaterials = FoliageType_ISM->OverrideMaterials;
+						bNeedsMarkRenderStateDirty = true;
+						bNeedsInvalidateLightingCache = true;
+						break;
+					}
+				}
+			}
+		}
+
+		Component->BodyInstance.CopyBodyInstancePropertiesFrom(&FoliageType->BodyInstance);
+
+		Component->SetCustomNavigableGeometry(FoliageType->CustomNavigableGeometry);
+
+		if (bNeedsInvalidateLightingCache)
+		{
+			Component->InvalidateLightingCache();
+		}
+
+		if (bNeedsMarkRenderStateDirty)
+		{
+			Component->MarkRenderStateDirty();
+		}
+	}
+}
+
+void FFoliageMeshInfo::AddInstances(AInstancedFoliageActor* InIFA, const UFoliageType* InSettings, const TSet<const FFoliageInstance*>& InNewInstances, bool RebuildFoliageTree)
+{
+	SCOPE_CYCLE_COUNTER(STAT_FoliageAddInstance);
+
+	InIFA->Modify();
+
+	if (Component == nullptr)
+	{
+		CreateNewComponent(InIFA, InSettings);
+		check(Component);
+	}
+	else
+	{
+		Component->InitPerInstanceRenderData(false);
+		Component->InvalidateLightingCache();
+	}
+
+	bool PreviousbAutoRebuildTreeOnInstanceChanges = Component->bAutoRebuildTreeOnInstanceChanges;
+	Component->bAutoRebuildTreeOnInstanceChanges = RebuildFoliageTree;
+
+	Instances.Reserve(Instances.Num() + InNewInstances.Num());
+	Component->PreAllocateInstancesMemory(InNewInstances.Num());
+
+	for (const FFoliageInstance* Instance : InNewInstances)
+	{
+		// Add the instance taking either a free slot or adding a new item.
+		int32 InstanceIndex = Instances.Add(*Instance);
+		FFoliageInstance& AddedInstance = Instances[InstanceIndex];
+
+		AddedInstance.BaseId = InIFA->InstanceBaseCache.AddInstanceBaseId(Instance->BaseComponent);
+
+		// Add the instance to the hash
+		AddToBaseHash(InstanceIndex);
+		InstanceHash->InsertInstance(AddedInstance.Location, InstanceIndex);
+
+		// Calculate transform for the instance
+		FTransform InstanceToWorld = Instance->GetInstanceWorldTransform();
+
+		// Add the instance to the component
+		Component->AddInstanceWorldSpace(InstanceToWorld);
+	}
+
+	CheckValid();
+
+	Component->bAutoRebuildTreeOnInstanceChanges = PreviousbAutoRebuildTreeOnInstanceChanges;
+}
+
+void FFoliageMeshInfo::AddInstance(AInstancedFoliageActor* InIFA, const UFoliageType* InSettings, const FFoliageInstance& InNewInstance, UActorComponent* InBaseComponent, bool RebuildFoliageTree)
+{
+	FFoliageInstance Instance = InNewInstance;
+	Instance.BaseId = InIFA->InstanceBaseCache.AddInstanceBaseId(InBaseComponent);
+	AddInstance(InIFA, InSettings, Instance, RebuildFoliageTree);
+}
+
+void FFoliageMeshInfo::AddInstance(AInstancedFoliageActor* InIFA, const UFoliageType* InSettings, const FFoliageInstance& InNewInstance, bool RebuildFoliageTree)
+{
+	SCOPE_CYCLE_COUNTER(STAT_FoliageAddInstance);
+
+	InIFA->Modify();
+
+	if (Component == nullptr)
+	{
+		CreateNewComponent(InIFA, InSettings);
+		check(Component);
+	}
+	else
+	{
+		Component->InitPerInstanceRenderData(false);
+		Component->InvalidateLightingCache();
+	}
+
+	bool PreviousbAutoRebuildTreeOnInstanceChanges = Component->bAutoRebuildTreeOnInstanceChanges;
+	Component->bAutoRebuildTreeOnInstanceChanges = RebuildFoliageTree;
+
+	// Add the instance taking either a free slot or adding a new item.
+	int32 InstanceIndex = Instances.Add(InNewInstance);
+	FFoliageInstance& AddedInstance = Instances[InstanceIndex];
+
+	// Add the instance to the hash
+	AddToBaseHash(InstanceIndex);
+	InstanceHash->InsertInstance(AddedInstance.Location, InstanceIndex);
+
+	// Calculate transform for the instance
+	FTransform InstanceToWorld = InNewInstance.GetInstanceWorldTransform();
+
+	// Add the instance to the component
+	Component->AddInstanceWorldSpace(InstanceToWorld);
+
+	CheckValid();
+
+	Component->bAutoRebuildTreeOnInstanceChanges = PreviousbAutoRebuildTreeOnInstanceChanges;
+}
+
+void FFoliageMeshInfo::RemoveInstances(AInstancedFoliageActor* InIFA, const TArray<int32>& InInstancesToRemove, bool RebuildFoliageTree)
+{
+	SCOPE_CYCLE_COUNTER(STAT_FoliageRemoveInstance);
+
+	if (InInstancesToRemove.Num() <= 0)
+	{
+		return;
+	}
+
+	check(Component);
+	InIFA->Modify();
+
+	bool PreviousbAutoRebuildTreeOnInstanceChanges = Component->bAutoRebuildTreeOnInstanceChanges;
+	Component->bAutoRebuildTreeOnInstanceChanges = false;
+
+	TSet<int32> InstancesToRemove;
+	InstancesToRemove.Append(InInstancesToRemove);
+
+	while (InstancesToRemove.Num() > 0)
+	{
+		// Get an item from the set for processing
+		auto It = InstancesToRemove.CreateConstIterator();
+		int32 InstanceIndex = *It;
+		int32 InstanceIndexToRemove = InstanceIndex;
+
+		FFoliageInstance& Instance = Instances[InstanceIndex];
+
+		// remove from hash
+		RemoveFromBaseHash(InstanceIndex);
+		InstanceHash->RemoveInstance(Instance.Location, InstanceIndex);
+
+		// remove from the component
+		Component->RemoveInstance(InstanceIndex);
+
+		// Remove it from the selection.
+		SelectedIndices.Remove(InstanceIndex);
+
+		// remove from instances array
+		Instances.RemoveAtSwap(InstanceIndex, 1, false);
+
+		// update hashes for swapped instance
+		if (InstanceIndex != Instances.Num() && Instances.Num() > 0)
+		{
+			// Instance hash
+			FFoliageInstance& SwappedInstance = Instances[InstanceIndex];
+			InstanceHash->RemoveInstance(SwappedInstance.Location, Instances.Num());
+			InstanceHash->InsertInstance(SwappedInstance.Location, InstanceIndex);
+
+			// Component hash
+			auto* InstanceSet = ComponentHash.Find(SwappedInstance.BaseId);
+			if (InstanceSet)
+			{
+				InstanceSet->Remove(Instances.Num());
+				InstanceSet->Add(InstanceIndex);
+			}
+
+			// Selection
+			if (SelectedIndices.Contains(Instances.Num()))
+			{
+				SelectedIndices.Remove(Instances.Num());
+				SelectedIndices.Add(InstanceIndex);
+			}
+
+			// Removal list
+			if (InstancesToRemove.Contains(Instances.Num()))
+			{
+				// The item from the end of the array that we swapped in to InstanceIndex is also on the list to remove.
+				// Remove the item at the end of the array and leave InstanceIndex in the removal list.
+				InstanceIndexToRemove = Instances.Num();
+			}
+		}
+
+		// Remove the removed item from the removal list
+		InstancesToRemove.Remove(InstanceIndexToRemove);
+	}
+
+	Instances.Shrink();
+		
+	Component->bAutoRebuildTreeOnInstanceChanges = PreviousbAutoRebuildTreeOnInstanceChanges;
+	Component->InvalidateLightingCache();
+
+	if (RebuildFoliageTree)
+	{
+		Component->BuildTreeIfOutdated(true, true);
+	}		
+
+	CheckValid();
+}
+
+void FFoliageMeshInfo::PreMoveInstances(AInstancedFoliageActor* InIFA, const TArray<int32>& InInstancesToMove)
+{
+	// Remove instances from the hash
+	for (TArray<int32>::TConstIterator It(InInstancesToMove); It; ++It)
+	{
+		int32 InstanceIndex = *It;
+		const FFoliageInstance& Instance = Instances[InstanceIndex];
+		InstanceHash->RemoveInstance(Instance.Location, InstanceIndex);
+	}
+}
+
+
+void FFoliageMeshInfo::PostUpdateInstances(AInstancedFoliageActor* InIFA, const TArray<int32>& InInstancesUpdated, bool bReAddToHash, bool InUpdateSelection)
+{
+	if (InInstancesUpdated.Num())
+	{
+		check(Component);
+
+		for (TArray<int32>::TConstIterator It(InInstancesUpdated); It; ++It)
+		{
+			int32 InstanceIndex = *It;
+			const FFoliageInstance& Instance = Instances[InstanceIndex];
+
+			FTransform InstanceToWorld = Instance.GetInstanceWorldTransform();
+
+			Component->UpdateInstanceTransform(InstanceIndex, InstanceToWorld, true);
+
+			// Re-add instance to the hash if requested
+			if (bReAddToHash)
+			{
+				InstanceHash->InsertInstance(Instance.Location, InstanceIndex);
+			}
+
+			// Reselect the instance to update the render update to include selection as by default it gets removed
+			if (InUpdateSelection)
+			{
+				Component->SelectInstance(true, InstanceIndex);
+			}
+		}
+
+		Component->InvalidateLightingCache();
+		Component->MarkRenderStateDirty();
+	}
+}
+
+void FFoliageMeshInfo::PostMoveInstances(AInstancedFoliageActor* InIFA, const TArray<int32>& InInstancesMoved)
+{
+	PostUpdateInstances(InIFA, InInstancesMoved, true, true);
+}
+
+void FFoliageMeshInfo::DuplicateInstances(AInstancedFoliageActor* InIFA, UFoliageType* InSettings, const TArray<int32>& InInstancesToDuplicate)
+{
+	if (Component != nullptr)
+	{
+		Component->bAutoRebuildTreeOnInstanceChanges = false;
+	}
+
+	for (int32 InstanceIndex : InInstancesToDuplicate)
+	{
+		const FFoliageInstance TempInstance = Instances[InstanceIndex];
+		AddInstance(InIFA, InSettings, TempInstance, false);
+	}
+
+	if (Component != nullptr)
+	{
+		Component->bAutoRebuildTreeOnInstanceChanges = true;
+		Component->BuildTreeIfOutdated(true, true);
+	}
+}
+
+/* Get the number of placed instances */
+int32 FFoliageMeshInfo::GetPlacedInstanceCount() const
+{
+	int32 PlacedInstanceCount = 0;
+
+	for (int32 i = 0; i < Instances.Num(); ++i)
+	{
+		if (!Instances[i].ProceduralGuid.IsValid())
+		{
+			++PlacedInstanceCount;
+		}
+	}
+
+	return PlacedInstanceCount;
+}
+
+void FFoliageMeshInfo::AddToBaseHash(int32 InstanceIndex)
+{
+	FFoliageInstance& Instance = Instances[InstanceIndex];
+	ComponentHash.FindOrAdd(Instance.BaseId).Add(InstanceIndex);
+}
+
+void FFoliageMeshInfo::RemoveFromBaseHash(int32 InstanceIndex)
+{
+	FFoliageInstance& Instance = Instances[InstanceIndex];
+
+	// Remove current base link
+	auto* InstanceSet = ComponentHash.Find(Instance.BaseId);
+	if (InstanceSet)
+	{
+		InstanceSet->Remove(InstanceIndex);
+		if (InstanceSet->Num() == 0)
+		{
+			// Remove the component from the component hash if this is the last instance.
+			ComponentHash.Remove(Instance.BaseId);
+		}
+	}
+}
+
+// Destroy existing clusters and reassign all instances to new clusters
+void FFoliageMeshInfo::ReallocateClusters(AInstancedFoliageActor* InIFA, UFoliageType* InSettings)
+{
+	if (Component != nullptr)
+	{
+		Component->ClearInstances();
+		Component->SetFlags(RF_Transactional);
+		Component->Modify();
+		Component->DestroyComponent();
+		Component = nullptr;
+	}
+
+	// Remove everything
+	TArray<FFoliageInstance> OldInstances;
+	Exchange(Instances, OldInstances);
+	InstanceHash->Empty();
+	ComponentHash.Empty();
+	SelectedIndices.Empty();
+
+	// Copy the UpdateGuid from the foliage type
+	FoliageTypeUpdateGuid = InSettings->UpdateGuid;
+
+	// Re-add
+	for (FFoliageInstance& Instance : OldInstances)
+	{
+		if ((Instance.Flags & FOLIAGE_InstanceDeleted) == 0)
+		{
+			AddInstance(InIFA, InSettings, Instance, false);
+		}
+	}
+
+	if (Component != nullptr)
+	{	
+		Component->BuildTreeIfOutdated(true, true);
+	}
+}
+
+void FFoliageMeshInfo::ReapplyInstancesToComponent()
+{
+	if (Component)
+	{
+		// clear the transactional flag if it was set prior to deleting the actor
+		Component->ClearFlags(RF_Transactional);
+
+		const bool bWasRegistered = Component->IsRegistered();
+		Component->UnregisterComponent();
+		Component->ClearInstances();
+		Component->InitPerInstanceRenderData(false);
+
+		Component->bAutoRebuildTreeOnInstanceChanges = false;
+
+		for (auto& Instance : Instances)
+		{
+			Component->AddInstanceWorldSpace(Instance.GetInstanceWorldTransform());
+		}
+
+		Component->bAutoRebuildTreeOnInstanceChanges = true;
+		Component->BuildTreeIfOutdated(false, true);
+
+		Component->ClearInstanceSelection();
+
+		if (SelectedIndices.Num())
+		{
+			for (int32 i : SelectedIndices)
+			{
+				Component->SelectInstance(true, i, 1);
+			}
+		}
+
+		if (bWasRegistered)
+		{
+			Component->RegisterComponent();
+		}
+	}
+}
+
+void FFoliageMeshInfo::GetInstancesInsideSphere(const FSphere& Sphere, TArray<int32>& OutInstances)
+{
+	auto TempInstances = InstanceHash->GetInstancesOverlappingBox(FBox::BuildAABB(Sphere.Center, FVector(Sphere.W)));
+	for (int32 Idx : TempInstances)
+	{
+		if (FSphere(Instances[Idx].Location, 0.f).IsInside(Sphere))
+		{
+			OutInstances.Add(Idx);
+		}
+	}
+}
+
+void FFoliageMeshInfo::GetInstanceAtLocation(const FVector& Location, int32& OutInstance, bool& bOutSucess)
+{
+	auto TempInstances = InstanceHash->GetInstancesOverlappingBox(FBox::BuildAABB(Location, FVector(KINDA_SMALL_NUMBER)));
+
+	float ShortestDistance = MAX_FLT;
+	OutInstance = -1;
+
+	for (int32 Idx : TempInstances)
+	{
+		FVector InstanceLocation = Instances[Idx].Location;
+		float DistanceSquared = FVector::DistSquared(InstanceLocation, Location);
+		if (DistanceSquared < ShortestDistance)
+		{
+			ShortestDistance = DistanceSquared;
+			OutInstance = Idx;
+		}
+	}
+
+	bOutSucess = OutInstance != -1;
+}
+
+// Returns whether or not there is are any instances overlapping the sphere specified
+bool FFoliageMeshInfo::CheckForOverlappingSphere(const FSphere& Sphere)
+{
+	auto TempInstances = InstanceHash->GetInstancesOverlappingBox(FBox::BuildAABB(Sphere.Center, FVector(Sphere.W)));
+	for (int32 Idx : TempInstances)
+	{
+		if (FSphere(Instances[Idx].Location, 0.f).IsInside(Sphere))
+		{
+			return true;
+		}
+	}
+	return false;
+}
+
+// Returns whether or not there is are any instances overlapping the instance specified, excluding the set of instances provided
+bool FFoliageMeshInfo::CheckForOverlappingInstanceExcluding(int32 TestInstanceIdx, float Radius, TSet<int32>& ExcludeInstances)
+{
+	FSphere Sphere(Instances[TestInstanceIdx].Location, Radius);
+
+	auto TempInstances = InstanceHash->GetInstancesOverlappingBox(FBox::BuildAABB(Sphere.Center, FVector(Sphere.W)));
+	for (int32 Idx : TempInstances)
+	{
+		if (Idx != TestInstanceIdx && !ExcludeInstances.Contains(Idx) && FSphere(Instances[Idx].Location, 0.f).IsInside(Sphere))
+		{
+			return true;
+		}
+	}
+	return false;
+}
+
+void FFoliageMeshInfo::SelectInstances(AInstancedFoliageActor* InIFA, bool bSelect)
+{
+	if (Component)
+	{
+		InIFA->Modify();
+
+		if (bSelect)
+		{
+			SelectedIndices.Reserve(Component->PerInstanceSMData.Num());
+
+			for (int32 i = 0; i < Component->PerInstanceSMData.Num(); ++i)
+			{
+				SelectedIndices.Add(i);
+			}
+
+			Component->SelectInstance(true, 0, Component->PerInstanceSMData.Num());
+		}
+		else
+		{
+			SelectedIndices.Empty();
+			Component->ClearInstanceSelection();
+		}
+
+		Component->MarkRenderStateDirty();
+	}
+}
+
+void FFoliageMeshInfo::SelectInstances(AInstancedFoliageActor* InIFA, bool bSelect, TArray<int32>& InInstances)
+{
+	if (InInstances.Num())
+	{
+		check(Component);
+		if (bSelect)
+		{
+			InIFA->Modify();
+
+			SelectedIndices.Reserve(InInstances.Num());
+
+			for (int32 i : InInstances)
+			{
+				SelectedIndices.Add(i);
+				Component->SelectInstance(true, i, 1);
+			}
+
+			Component->MarkRenderStateDirty();
+		}
+		else
+		{
+			InIFA->Modify();
+
+			for (int32 i : InInstances)
+			{
+				SelectedIndices.Remove(i);
+			}
+
+			if (Component->SelectedInstances.Num() > 0)
+			{
+				for (int32 i : InInstances)
+				{
+					Component->SelectInstance(false, i, 1);
+				}
+
+				Component->MarkRenderStateDirty();
+			}
+		}
+	}
+}
+
+#endif	//WITH_EDITOR
+
+//
+// AInstancedFoliageActor
+//
+AInstancedFoliageActor::AInstancedFoliageActor(const FObjectInitializer& ObjectInitializer)
+	: Super(ObjectInitializer)
+{
+	USceneComponent* SceneComponent = CreateDefaultSubobject<USceneComponent>(TEXT("RootComponent0"));
+	RootComponent = SceneComponent;
+	RootComponent->Mobility = EComponentMobility::Static;
+
+	SetActorEnableCollision(true);
+#if WITH_EDITORONLY_DATA
+	bListedInSceneOutliner = false;
+#endif // WITH_EDITORONLY_DATA
+	PrimaryActorTick.bCanEverTick = false;
+}
+
+AInstancedFoliageActor* AInstancedFoliageActor::GetInstancedFoliageActorForCurrentLevel(UWorld* InWorld, bool bCreateIfNone)
+{
+	return GetInstancedFoliageActorForLevel(InWorld->GetCurrentLevel(), bCreateIfNone);
+}
+
+AInstancedFoliageActor* AInstancedFoliageActor::GetInstancedFoliageActorForLevel(ULevel* InLevel, bool bCreateIfNone /* = false */)
+{
+	AInstancedFoliageActor* IFA = nullptr;
+	if (InLevel)
+	{
+		IFA = InLevel->InstancedFoliageActor.Get();
+
+		if (!IFA && bCreateIfNone)
+		{
+			FActorSpawnParameters SpawnParams;
+			SpawnParams.OverrideLevel = InLevel;
+			IFA = InLevel->GetWorld()->SpawnActor<AInstancedFoliageActor>(SpawnParams);
+			InLevel->InstancedFoliageActor = IFA;
+		}
+	}
+
+	return IFA;
+}
+
+
+int32 AInstancedFoliageActor::GetOverlappingSphereCount(const UFoliageType* FoliageType, const FSphere& Sphere) const
+{
+	if (const FFoliageMeshInfo* MeshInfo = FindMesh(FoliageType))
+	{
+		if (MeshInfo->Component && MeshInfo->Component->IsTreeFullyBuilt())
+		{
+			return MeshInfo->Component->GetOverlappingSphereCount(Sphere);
+		}
+	}
+
+	return 0;
+}
+
+
+int32 AInstancedFoliageActor::GetOverlappingBoxCount(const UFoliageType* FoliageType, const FBox& Box) const
+{
+	if (const FFoliageMeshInfo* MeshInfo = FindMesh(FoliageType))
+	{
+		if (MeshInfo->Component && MeshInfo->Component->IsTreeFullyBuilt())
+		{
+			return MeshInfo->Component->GetOverlappingBoxCount(Box);
+		}
+	}
+
+	return 0;
+}
+
+
+void AInstancedFoliageActor::GetOverlappingBoxTransforms(const UFoliageType* FoliageType, const FBox& Box, TArray<FTransform>& OutTransforms) const
+{
+	if (const FFoliageMeshInfo* MeshInfo = FindMesh(FoliageType))
+	{
+		if (MeshInfo->Component && MeshInfo->Component->IsTreeFullyBuilt())
+		{
+			MeshInfo->Component->GetOverlappingBoxTransforms(Box, OutTransforms);
+		}
+	}
+}
+
+void AInstancedFoliageActor::GetOverlappingMeshCounts(const FSphere& Sphere, TMap<UStaticMesh*, int32>& OutCounts) const
+{
+	for (auto& MeshPair : FoliageMeshes)
+	{
+		FFoliageMeshInfo const* MeshInfo = &*MeshPair.Value;
+
+		if (MeshInfo && MeshInfo->Component)
+		{
+			int32 const Count = MeshInfo->Component->GetOverlappingSphereCount(Sphere);
+			if (Count > 0)
+			{
+				UStaticMesh* const Mesh = MeshInfo->Component->GetStaticMesh();
+				int32& StoredCount = OutCounts.FindOrAdd(Mesh);
+				StoredCount += Count;
+			}
+		}
+	}
+}
+
+
+
+UFoliageType* AInstancedFoliageActor::GetLocalFoliageTypeForMesh(const UStaticMesh* InMesh, FFoliageMeshInfo** OutMeshInfo)
+{
+	UFoliageType* ReturnType = nullptr;
+	FFoliageMeshInfo* MeshInfo = nullptr;
+
+	for (auto& MeshPair : FoliageMeshes)
+	{
+		UFoliageType* FoliageType = MeshPair.Key;
+		// Check that the type is neither an asset nor blueprint instance
+		if (FoliageType && FoliageType->GetStaticMesh() == InMesh && !FoliageType->IsAsset() && !FoliageType->GetClass()->ClassGeneratedBy)
+		{
+			ReturnType = FoliageType;
+			MeshInfo = &*MeshPair.Value;
+			break;
+		}
+	}
+
+	if (OutMeshInfo)
+	{
+		*OutMeshInfo = MeshInfo;
+	}
+
+	return ReturnType;
+}
+
+void AInstancedFoliageActor::GetAllFoliageTypesForMesh(const UStaticMesh* InMesh, TArray<const UFoliageType*>& OutFoliageTypes)
+{
+	for (auto& MeshPair : FoliageMeshes)
+	{
+		UFoliageType* FoliageType = MeshPair.Key;
+		if (FoliageType && FoliageType->GetStaticMesh() == InMesh)
+		{
+			OutFoliageTypes.Add(FoliageType);
+		}
+	}
+}
+
+
+FFoliageMeshInfo* AInstancedFoliageActor::FindFoliageTypeOfClass(TSubclassOf<UFoliageType_InstancedStaticMesh> Class)
+{
+	FFoliageMeshInfo* MeshInfo = nullptr;
+
+	for (auto& MeshPair : FoliageMeshes)
+	{
+		UFoliageType* FoliageType = MeshPair.Key;
+		if (FoliageType && FoliageType->GetClass() == Class)
+		{
+			MeshInfo = &MeshPair.Value.Get();
+			break;
+		}
+	}
+
+	return MeshInfo;
+}
+
+FFoliageMeshInfo* AInstancedFoliageActor::FindMesh(const UFoliageType* InType)
+{
+	TUniqueObj<FFoliageMeshInfo>* MeshInfoEntry = FoliageMeshes.Find(InType);
+	FFoliageMeshInfo* MeshInfo = MeshInfoEntry ? &MeshInfoEntry->Get() : nullptr;
+	return MeshInfo;
+}
+
+const FFoliageMeshInfo* AInstancedFoliageActor::FindMesh(const UFoliageType* InType) const
+{
+	const TUniqueObj<FFoliageMeshInfo>* MeshInfoEntry = FoliageMeshes.Find(InType);
+	const FFoliageMeshInfo* MeshInfo = MeshInfoEntry ? &MeshInfoEntry->Get() : nullptr;
+	return MeshInfo;
+}
+
+
+#if WITH_EDITOR
+void AInstancedFoliageActor::MoveInstancesForMovedComponent(UActorComponent* InComponent)
+{
+	const auto BaseId = InstanceBaseCache.GetInstanceBaseId(InComponent);
+	if (BaseId == FFoliageInstanceBaseCache::InvalidBaseId)
+	{
+		return;
+	}
+
+	const auto CurrentBaseInfo = InstanceBaseCache.GetInstanceBaseInfo(BaseId);
+
+	// Found an invalid base so don't try to move instances
+	if (!CurrentBaseInfo.BasePtr.IsValid())
+	{
+		return;
+	}
+
+	bool bFirst = true;
+	const auto NewBaseInfo = InstanceBaseCache.UpdateInstanceBaseInfoTransform(InComponent);
+
+	FMatrix DeltaTransfrom =
+		FTranslationMatrix(-CurrentBaseInfo.CachedLocation) *
+		FInverseRotationMatrix(CurrentBaseInfo.CachedRotation) *
+		FScaleMatrix(NewBaseInfo.CachedDrawScale / CurrentBaseInfo.CachedDrawScale) *
+		FRotationMatrix(NewBaseInfo.CachedRotation) *
+		FTranslationMatrix(NewBaseInfo.CachedLocation);
+
+	for (auto& MeshPair : FoliageMeshes)
+	{
+		FFoliageMeshInfo& MeshInfo = *MeshPair.Value;
+		const auto* InstanceSet = MeshInfo.ComponentHash.Find(BaseId);
+		if (InstanceSet && InstanceSet->Num())
+		{
+			if (bFirst)
+			{
+				bFirst = false;
+				Modify();
+			}
+
+			check(MeshInfo.Component);
+			MeshInfo.Component->bAutoRebuildTreeOnInstanceChanges = false;
+
+			for (int32 InstanceIndex : *InstanceSet)
+			{
+				FFoliageInstance& Instance = MeshInfo.Instances[InstanceIndex];
+
+				MeshInfo.InstanceHash->RemoveInstance(Instance.Location, InstanceIndex);
+
+				// Apply change
+				FMatrix NewTransform =
+					FRotationMatrix(Instance.Rotation) *
+					FTranslationMatrix(Instance.Location) *
+					DeltaTransfrom;
+
+				// Extract rotation and position
+				Instance.Location = NewTransform.GetOrigin();
+				Instance.Rotation = NewTransform.Rotator();
+
+				// Apply render data
+				MeshInfo.Component->UpdateInstanceTransform(InstanceIndex, Instance.GetInstanceWorldTransform(), true, true);
+
+				// Re-add the new instance location to the hash
+				MeshInfo.InstanceHash->InsertInstance(Instance.Location, InstanceIndex);
+			}
+
+			MeshInfo.Component->bAutoRebuildTreeOnInstanceChanges = true;
+			MeshInfo.Component->BuildTreeIfOutdated(true, false);
+		}
+	}
+}
+
+void AInstancedFoliageActor::DeleteInstancesForComponent(UActorComponent* InComponent)
+{
+	const auto BaseId = InstanceBaseCache.GetInstanceBaseId(InComponent);
+	// Instances with empty base has BaseId==InvalidBaseId, we should not delete these
+	if (BaseId == FFoliageInstanceBaseCache::InvalidBaseId)
+	{
+		return;
+	}
+
+	for (auto& MeshPair : FoliageMeshes)
+	{
+		FFoliageMeshInfo& MeshInfo = *MeshPair.Value;
+		const auto* InstanceSet = MeshInfo.ComponentHash.Find(BaseId);
+		if (InstanceSet)
+		{
+			MeshInfo.RemoveInstances(this, InstanceSet->Array(), true);
+		}
+	}
+}
+
+void AInstancedFoliageActor::DeleteInstancesForComponent(UActorComponent* InComponent, const UFoliageType* FoliageType)
+{
+	const auto BaseId = InstanceBaseCache.GetInstanceBaseId(InComponent);
+	// Instances with empty base has BaseId==InvalidBaseId, we should not delete these
+	if (BaseId == FFoliageInstanceBaseCache::InvalidBaseId)
+	{
+		return;
+	}
+
+	FFoliageMeshInfo* MeshInfo = FindMesh(FoliageType);
+	if (MeshInfo)
+	{
+		const auto* InstanceSet = MeshInfo->ComponentHash.Find(BaseId);
+		if (InstanceSet)
+		{
+			MeshInfo->RemoveInstances(this, InstanceSet->Array(), true);
+		}
+	}
+}
+
+void AInstancedFoliageActor::DeleteInstancesForComponent(UWorld* InWorld, UActorComponent* InComponent)
+{
+	for (TActorIterator<AInstancedFoliageActor> It(InWorld); It; ++It)
+	{
+		AInstancedFoliageActor* IFA = (*It);
+		IFA->Modify();
+		IFA->DeleteInstancesForComponent(InComponent);
+	}
+}
+
+void AInstancedFoliageActor::DeleteInstancesForProceduralFoliageComponent(const UProceduralFoliageComponent* ProceduralFoliageComponent, bool InRebuildTree)
+{
+	const FGuid& ProceduralGuid = ProceduralFoliageComponent->GetProceduralGuid();
+	for (auto& MeshPair : FoliageMeshes)
+	{
+		FFoliageMeshInfo& MeshInfo = *MeshPair.Value;
+		TArray<int32> InstancesToRemove;
+		for (int32 InstanceIdx = 0; InstanceIdx < MeshInfo.Instances.Num(); InstanceIdx++)
+		{
+			if (MeshInfo.Instances[InstanceIdx].ProceduralGuid == ProceduralGuid)
+			{
+				InstancesToRemove.Add(InstanceIdx);
+			}
+		}
+
+		if (InstancesToRemove.Num())
+		{
+			MeshInfo.RemoveInstances(this, InstancesToRemove, InRebuildTree);
+		}
+	}
+
+	// Clean up dead cross-level references
+	FFoliageInstanceBaseCache::CompactInstanceBaseCache(this);
+}
+
+bool AInstancedFoliageActor::ContainsInstancesFromProceduralFoliageComponent(const UProceduralFoliageComponent* ProceduralFoliageComponent)
+{
+	const FGuid& ProceduralGuid = ProceduralFoliageComponent->GetProceduralGuid();
+	for (auto& MeshPair : FoliageMeshes)
+	{
+		FFoliageMeshInfo& MeshInfo = *MeshPair.Value;
+		TArray<int32> InstancesToRemove;
+		for (int32 InstanceIdx = 0; InstanceIdx < MeshInfo.Instances.Num(); InstanceIdx++)
+		{
+			if (MeshInfo.Instances[InstanceIdx].ProceduralGuid == ProceduralGuid)
+			{
+				// The procedural component is responsible for an instance
+				return true;
+			}
+		}
+	}
+
+	return false;
+}
+
+void AInstancedFoliageActor::MoveInstancesForComponentToCurrentLevel(UActorComponent* InComponent)
+{
+	if (!HasFoliageAttached(InComponent))
+	{
+		// Quit early if there are no foliage instances painted on this component
+		return;
+	}
+
+	UWorld* InWorld = InComponent->GetWorld();
+	AInstancedFoliageActor* NewIFA = AInstancedFoliageActor::GetInstancedFoliageActorForCurrentLevel(InWorld, true);
+	NewIFA->Modify();
+
+	for (TActorIterator<AInstancedFoliageActor> It(InWorld); It; ++It)
+	{
+		AInstancedFoliageActor* IFA = (*It);
+
+		const auto SourceBaseId = IFA->InstanceBaseCache.GetInstanceBaseId(InComponent);
+		if (SourceBaseId != FFoliageInstanceBaseCache::InvalidBaseId && IFA != NewIFA)
+		{
+			IFA->Modify();
+
+			for (auto& MeshPair : IFA->FoliageMeshes)
+			{
+				FFoliageMeshInfo& MeshInfo = *MeshPair.Value;
+				UFoliageType* FoliageType = MeshPair.Key;
+
+				const auto* InstanceSet = MeshInfo.ComponentHash.Find(SourceBaseId);
+				if (InstanceSet)
+				{
+					// Duplicate the foliage type if it's not shared
+					FFoliageMeshInfo* TargetMeshInfo = nullptr;
+					UFoliageType* TargetFoliageType = NewIFA->AddFoliageType(FoliageType, &TargetMeshInfo);
+
+					// Add the foliage to the new level
+					for (int32 InstanceIndex : *InstanceSet)
+					{
+						TargetMeshInfo->AddInstance(NewIFA, TargetFoliageType, MeshInfo.Instances[InstanceIndex], InComponent, false);
+					}
+
+					TargetMeshInfo->Component->BuildTreeIfOutdated(true, true);
+
+					// Remove from old level
+					MeshInfo.RemoveInstances(IFA, InstanceSet->Array(), true);
+				}
+			}
+		}
+	}
+}
+
+void AInstancedFoliageActor::MoveInstancesToNewComponent(UPrimitiveComponent* InOldComponent, const FBox& InBoxWithInstancesToMove, UPrimitiveComponent* InNewComponent)
+{	
+	const auto OldBaseId = InstanceBaseCache.GetInstanceBaseId(InOldComponent);
+	if (OldBaseId == FFoliageInstanceBaseCache::InvalidBaseId)
+	{
+		// This foliage actor has no instances with specified base
+		return;
+	}
+
+	AInstancedFoliageActor* TargetIFA = AInstancedFoliageActor::GetInstancedFoliageActorForLevel(InNewComponent->GetTypedOuter<ULevel>(), true);
+	TArray<int32> InstancesToMove;
+
+	for (auto& MeshPair : FoliageMeshes)
+	{
+		FFoliageMeshInfo& MeshInfo = *MeshPair.Value;
+		
+		if (MeshInfo.Component != nullptr)
+		{
+			InstancesToMove = MeshInfo.Component->GetInstancesOverlappingBox(InBoxWithInstancesToMove);
+		}
+
+		FFoliageMeshInfo* TargetMeshInfo = nullptr;
+		UFoliageType* TargetFoliageType = TargetIFA->AddFoliageType(MeshPair.Key, &TargetMeshInfo);
+
+		// Add the foliage to the new level
+		for (int32 InstanceIndex : InstancesToMove)
+		{
+			FFoliageInstance NewInstance = MeshInfo.Instances[InstanceIndex];
+			TargetMeshInfo->AddInstance(TargetIFA, TargetFoliageType, NewInstance, InNewComponent, false);
+		}
+
+		if (TargetMeshInfo->Component != nullptr)
+		{
+			TargetMeshInfo->Component->BuildTreeIfOutdated(true, true);
+		}
+
+		// Remove from old level
+		MeshInfo.RemoveInstances(this, InstancesToMove, true);
+	}
+}
+
+void AInstancedFoliageActor::MoveInstancesToNewComponent(UPrimitiveComponent* InOldComponent, UPrimitiveComponent* InNewComponent)
+{
+	AInstancedFoliageActor* TargetIFA = AInstancedFoliageActor::GetInstancedFoliageActorForLevel(InNewComponent->GetTypedOuter<ULevel>(), true);
+
+	const auto OldBaseId = this->InstanceBaseCache.GetInstanceBaseId(InOldComponent);
+	if (OldBaseId == FFoliageInstanceBaseCache::InvalidBaseId)
+	{
+		// This foliage actor has no any instances with specified base
+		return;
+	}
+
+	const auto NewBaseId = TargetIFA->InstanceBaseCache.AddInstanceBaseId(InNewComponent);
+
+	for (auto& MeshPair : FoliageMeshes)
+	{
+		FFoliageMeshInfo& MeshInfo = *MeshPair.Value;
+
+		TSet<int32> InstanceSet;
+		if (MeshInfo.ComponentHash.RemoveAndCopyValue(OldBaseId, InstanceSet) && InstanceSet.Num())
+		{
+			// For same FoliageActor can just remap the instances, otherwise we have to do a more complex move
+			if (TargetIFA == this)
+			{
+				// Update the instances
+				for (int32 InstanceIndex : InstanceSet)
+				{
+					MeshInfo.Instances[InstanceIndex].BaseId = NewBaseId;
+				}
+
+				// Update the hash
+				MeshInfo.ComponentHash.Add(NewBaseId, MoveTemp(InstanceSet));
+			}
+			else
+			{
+				FFoliageMeshInfo* TargetMeshInfo = nullptr;
+				UFoliageType* TargetFoliageType = TargetIFA->AddFoliageType(MeshPair.Key, &TargetMeshInfo);
+
+				// Add the foliage to the new level
+				for (int32 InstanceIndex : InstanceSet)
+				{
+					FFoliageInstance NewInstance = MeshInfo.Instances[InstanceIndex];
+					NewInstance.BaseId = NewBaseId;
+					TargetMeshInfo->AddInstance(TargetIFA, TargetFoliageType, NewInstance, false);
+				}
+
+				if (TargetMeshInfo->Component != nullptr)
+				{
+					TargetMeshInfo->Component->BuildTreeIfOutdated(true, true);
+				}
+
+				// Remove from old level
+				MeshInfo.RemoveInstances(this, InstanceSet.Array(), true);
+			}
+		}
+	}
+}
+
+void AInstancedFoliageActor::MoveInstancesToNewComponent(UWorld* InWorld, UPrimitiveComponent* InOldComponent, UPrimitiveComponent* InNewComponent)
+{
+	for (TActorIterator<AInstancedFoliageActor> It(InWorld); It; ++It)
+	{
+		AInstancedFoliageActor* IFA = (*It);
+		IFA->MoveInstancesToNewComponent(InOldComponent, InNewComponent);
+	}
+}
+
+void AInstancedFoliageActor::MoveInstancesToNewComponent(UWorld* InWorld, UPrimitiveComponent* InOldComponent, const FBox& InBoxWithInstancesToMove, UPrimitiveComponent* InNewComponent)
+{
+	for (TActorIterator<AInstancedFoliageActor> It(InWorld); It; ++It)
+	{
+		AInstancedFoliageActor* IFA = (*It);
+		IFA->MoveInstancesToNewComponent(InOldComponent, InBoxWithInstancesToMove, InNewComponent);
+	}
+}
+
+void AInstancedFoliageActor::MoveInstancesToLevel(ULevel* InTargetLevel, TSet<int32>& InInstanceList, FFoliageMeshInfo* InCurrentMeshInfo, UFoliageType* InFoliageType)
+{
+	if (InTargetLevel == GetLevel())
+	{
+		return;
+	}
+
+	AInstancedFoliageActor* TargetIFA = GetInstancedFoliageActorForLevel(InTargetLevel, /*bCreateIfNone*/ true);
+
+	Modify();
+	TargetIFA->Modify();
+
+	// Do move
+	FFoliageMeshInfo* TargetMeshInfo = nullptr;
+	UFoliageType* TargetFoliageType = TargetIFA->AddFoliageType(InFoliageType, &TargetMeshInfo);
+
+	// Add selected instances to the target actor
+	for (int32 InstanceIndex : InInstanceList)
+	{
+		FFoliageInstance& Instance = InCurrentMeshInfo->Instances[InstanceIndex];
+		TargetMeshInfo->AddInstance(TargetIFA, TargetFoliageType, Instance, InstanceBaseCache.GetInstanceBasePtr(Instance.BaseId).Get(), false);
+	}
+
+	if (TargetMeshInfo->Component != nullptr)
+	{
+		TargetMeshInfo->Component->BuildTreeIfOutdated(true, true);
+	}
+
+	// Remove selected instances from this actor
+	InCurrentMeshInfo->RemoveInstances(this, InInstanceList.Array(), true);
+}
+
+void AInstancedFoliageActor::MoveSelectedInstancesToLevel(ULevel* InTargetLevel)
+{
+	if (InTargetLevel == GetLevel() || !HasSelectedInstances())
+	{
+		return;
+	}
+
+	for (auto& MeshPair : FoliageMeshes)
+	{
+		FFoliageMeshInfo& MeshInfo = *MeshPair.Value;
+		UFoliageType* FoliageType = MeshPair.Key;
+
+		MoveInstancesToLevel(InTargetLevel, MeshInfo.SelectedIndices, &MeshInfo, FoliageType);
+	}
+}
+
+void AInstancedFoliageActor::MoveAllInstancesToLevel(ULevel* InTargetLevel)
+{
+	if (InTargetLevel == GetLevel())
+	{
+		return;
+	}
+
+	for (auto& MeshPair : FoliageMeshes)
+	{
+		FFoliageMeshInfo& MeshInfo = *MeshPair.Value;
+		UFoliageType* FoliageType = MeshPair.Key;
+
+		TSet<int32> instancesList;
+
+		for (int32 i = 0; i < MeshInfo.Instances.Num(); ++i)
+		{
+			instancesList.Add(i);
+		}
+
+		MoveInstancesToLevel(InTargetLevel, instancesList, &MeshInfo, FoliageType);
+	}
+}
+
+TMap<UFoliageType*, TArray<const FFoliageInstancePlacementInfo*>> AInstancedFoliageActor::GetInstancesForComponent(UActorComponent* InComponent)
+{
+	TMap<UFoliageType*, TArray<const FFoliageInstancePlacementInfo*>> Result;
+	const auto BaseId = InstanceBaseCache.GetInstanceBaseId(InComponent);
+
+	if (BaseId != FFoliageInstanceBaseCache::InvalidBaseId)
+	{
+		for (auto& MeshPair : FoliageMeshes)
+		{
+			const FFoliageMeshInfo& MeshInfo = *MeshPair.Value;
+			const auto* InstanceSet = MeshInfo.ComponentHash.Find(BaseId);
+			if (InstanceSet)
+			{
+				TArray<const FFoliageInstancePlacementInfo*>& Array = Result.Add(MeshPair.Key, TArray<const FFoliageInstancePlacementInfo*>());
+				Array.Empty(InstanceSet->Num());
+
+				for (int32 InstanceIndex : *InstanceSet)
+				{
+					const FFoliageInstancePlacementInfo* Instance = &MeshInfo.Instances[InstanceIndex];
+					Array.Add(Instance);
+				}
+			}
+		}
+	}
+
+	return Result;
+}
+
+FFoliageMeshInfo* AInstancedFoliageActor::FindOrAddMesh(UFoliageType* InType)
+{
+	TUniqueObj<FFoliageMeshInfo>* MeshInfoEntry = FoliageMeshes.Find(InType);
+	FFoliageMeshInfo* MeshInfo = MeshInfoEntry ? &MeshInfoEntry->Get() : AddMesh(InType);
+	return MeshInfo;
+}
+
+
+void UpdateSettingsBounds(const UStaticMesh* InMesh, UFoliageType_InstancedStaticMesh* Settings)
+{
+	const FBoxSphereBounds MeshBounds = InMesh->GetBounds();
+
+	Settings->MeshBounds = MeshBounds;
+
+	// Make bottom only bound
+	FBox LowBound = MeshBounds.GetBox();
+	LowBound.Max.Z = LowBound.Min.Z + (LowBound.Max.Z - LowBound.Min.Z) * 0.1f;
+
+	float MinX = FLT_MAX, MaxX = FLT_MIN, MinY = FLT_MAX, MaxY = FLT_MIN;
+	Settings->LowBoundOriginRadius = FVector::ZeroVector;
+
+	if (InMesh->RenderData)
+	{
+		FPositionVertexBuffer& PositionVertexBuffer = InMesh->RenderData->LODResources[0].VertexBuffers.PositionVertexBuffer;
+		for (uint32 Index = 0; Index < PositionVertexBuffer.GetNumVertices(); ++Index)
+		{
+			const FVector& Pos = PositionVertexBuffer.VertexPosition(Index);
+			if (Pos.Z < LowBound.Max.Z)
+			{
+				MinX = FMath::Min(MinX, Pos.X);
+				MinY = FMath::Min(MinY, Pos.Y);
+				MaxX = FMath::Max(MaxX, Pos.X);
+				MaxY = FMath::Max(MaxY, Pos.Y);
+			}
+		}
+	}
+
+	Settings->LowBoundOriginRadius = FVector((MinX + MaxX), (MinY + MaxY), FMath::Sqrt(FMath::Square(MaxX - MinX) + FMath::Square(MaxY - MinY))) * 0.5f;
+}
+
+UFoliageType* AInstancedFoliageActor::AddFoliageType(const UFoliageType* InType, FFoliageMeshInfo** OutInfo)
+{
+	FFoliageMeshInfo* MeshInfo = nullptr;
+	UFoliageType* FoliageType = const_cast<UFoliageType*>(InType);
+
+	if (FoliageType->GetOuter() == this || FoliageType->IsAsset())
+	{
+		auto ExistingMeshInfo = FoliageMeshes.Find(FoliageType);
+		if (!ExistingMeshInfo)
+		{
+			Modify();
+			MeshInfo = &FoliageMeshes.Add(FoliageType).Get();
+		}
+		else
+		{
+			MeshInfo = &ExistingMeshInfo->Get();
+		}
+	}
+	else if (FoliageType->GetClass()->ClassGeneratedBy)
+	{
+		// Foliage type blueprint
+		FFoliageMeshInfo* ExistingMeshInfo = FindFoliageTypeOfClass(FoliageType->GetClass());
+		if (!ExistingMeshInfo)
+		{
+			Modify();
+			FoliageType = DuplicateObject<UFoliageType>(InType, this);
+			MeshInfo = &FoliageMeshes.Add(FoliageType).Get();
+		}
+		else
+		{
+			MeshInfo = ExistingMeshInfo;
+		}
+	}
+	else
+	{
+		// Unique meshes only
+		// Multiple entries for same static mesh can be added using FoliageType as an asset
+		FoliageType = GetLocalFoliageTypeForMesh(FoliageType->GetStaticMesh(), &MeshInfo);
+		if (FoliageType == nullptr)
+		{
+			Modify();
+			FoliageType = DuplicateObject<UFoliageType>(InType, this);
+			MeshInfo = &FoliageMeshes.Add(FoliageType).Get();
+		}
+	}
+
+	if (OutInfo)
+	{
+		*OutInfo = MeshInfo;
+	}
+
+	return FoliageType;
+}
+
+FFoliageMeshInfo* AInstancedFoliageActor::AddMesh(UStaticMesh* InMesh, UFoliageType** OutSettings, const UFoliageType_InstancedStaticMesh* DefaultSettings)
+{
+	check(GetLocalFoliageTypeForMesh(InMesh) == nullptr);
+
+	MarkPackageDirty();
+
+	UFoliageType_InstancedStaticMesh* Settings = nullptr;
+#if WITH_EDITORONLY_DATA
+	if (DefaultSettings)
+	{
+		// TODO: Can't we just use this directly?
+		FObjectDuplicationParameters DuplicationParameters(const_cast<UFoliageType_InstancedStaticMesh*>(DefaultSettings), this);
+		DuplicationParameters.ApplyFlags = RF_Transactional;
+		Settings = CastChecked<UFoliageType_InstancedStaticMesh>(StaticDuplicateObjectEx(DuplicationParameters));
+	}
+	else
+#endif
+	{
+		Settings = NewObject<UFoliageType_InstancedStaticMesh>(this, NAME_None, RF_Transactional);
+	}
+	Settings->Mesh = InMesh;
+
+	FFoliageMeshInfo* MeshInfo = AddMesh(Settings);
+	UpdateSettingsBounds(InMesh, Settings);
+
+	if (OutSettings)
+	{
+		*OutSettings = Settings;
+	}
+
+	return MeshInfo;
+}
+
+FFoliageMeshInfo* AInstancedFoliageActor::AddMesh(UFoliageType* InType)
+{
+	check(FoliageMeshes.Find(InType) == nullptr);
+
+	Modify();
+
+	FFoliageMeshInfo* MeshInfo = &*FoliageMeshes.Add(InType);
+	MeshInfo->FoliageTypeUpdateGuid = InType->UpdateGuid;
+	InType->IsSelected = true;
+
+	return MeshInfo;
+}
+
+void AInstancedFoliageActor::RemoveFoliageType(UFoliageType** InFoliageTypes, int32 Num)
+{
+	Modify();
+	UnregisterAllComponents();
+
+	// Remove all components for this mesh from the Components array.
+	for (int32 FoliageTypeIdx = 0; FoliageTypeIdx < Num; ++FoliageTypeIdx)
+	{
+		const UFoliageType* FoliageType = InFoliageTypes[FoliageTypeIdx];
+		FFoliageMeshInfo* MeshInfo = FindMesh(FoliageType);
+		if (MeshInfo)
+		{
+			if (MeshInfo->Component)
+			{
+				if (MeshInfo->Component->GetStaticMesh() != nullptr)
+				{
+					MeshInfo->Component->GetStaticMesh()->GetOnExtendedBoundsChanged().RemoveAll(MeshInfo);
+				}
+
+				MeshInfo->Component->ClearInstances();
+				MeshInfo->Component->SetFlags(RF_Transactional);
+				MeshInfo->Component->Modify();
+				MeshInfo->Component->DestroyComponent();
+				MeshInfo->Component = nullptr;
+			}
+
+			FoliageMeshes.Remove(FoliageType);
+		}
+	}
+
+	RegisterAllComponents();
+}
+
+void AInstancedFoliageActor::SelectInstance(UInstancedStaticMeshComponent* InComponent, int32 InInstanceIndex, bool bToggle)
+{
+	Modify();
+
+	// If we're not toggling, we need to first deselect everything else
+	if (!bToggle)
+	{
+		for (auto& MeshPair : FoliageMeshes)
+		{
+			FFoliageMeshInfo& MeshInfo = *MeshPair.Value;
+
+			if (MeshInfo.Instances.Num() > 0)
+			{
+				check(MeshInfo.Component);
+				MeshInfo.Component->ClearInstanceSelection();
+				MeshInfo.Component->MarkRenderStateDirty();
+				MeshInfo.SelectedIndices.Empty();
+			}
+		}
+	}
+
+	if (InComponent)
+	{
+		UFoliageType* Type = nullptr;
+		FFoliageMeshInfo* MeshInfo = nullptr;
+
+		for (auto& MeshPair : FoliageMeshes)
+		{
+			if (MeshPair.Value->Component == InComponent)
+			{
+				Type = MeshPair.Key;
+				MeshInfo = &MeshPair.Value.Get();
+				break;
+			}
+		}
+
+		if (MeshInfo)
+		{
+			bool bIsSelected = MeshInfo->SelectedIndices.Contains(InInstanceIndex);
+
+			// Deselect if it's already selected.
+			if (InInstanceIndex < InComponent->SelectedInstances.Num())
+			{
+				InComponent->SelectInstance(false, InInstanceIndex, 1);
+				InComponent->MarkRenderStateDirty();
+			}
+
+			if (bIsSelected)
+			{
+				MeshInfo->SelectedIndices.Remove(InInstanceIndex);
+			}
+
+			if (!bToggle || !bIsSelected)
+			{
+				// Add the selection
+				InComponent->SelectInstance(true, InInstanceIndex, 1);
+				InComponent->MarkRenderStateDirty();
+
+				MeshInfo->SelectedIndices.Add(InInstanceIndex);
+			}
+		}
+	}
+}
+
+bool AInstancedFoliageActor::HasSelectedInstances() const
+{
+	for (const auto& MeshPair : FoliageMeshes)
+	{
+		if (MeshPair.Value->SelectedIndices.Num() > 0)
+		{
+			return true;
+		}
+	}
+
+	return false;
+}
+
+TMap<UFoliageType*, FFoliageMeshInfo*> AInstancedFoliageActor::GetAllInstancesFoliageType()
+{
+	TMap<UFoliageType*, FFoliageMeshInfo*> InstanceFoliageTypes;
+
+	for (auto& MeshPair : FoliageMeshes)
+	{
+		InstanceFoliageTypes.Add(MeshPair.Key, &MeshPair.Value.Get());
+	}
+
+	return InstanceFoliageTypes;
+}
+
+TMap<UFoliageType*, FFoliageMeshInfo*> AInstancedFoliageActor::GetSelectedInstancesFoliageType()
+{
+	TMap<UFoliageType*, FFoliageMeshInfo*> SelectedInstanceFoliageTypes;
+
+	for (auto& MeshPair : FoliageMeshes)
+	{
+		if (MeshPair.Value->SelectedIndices.Num() > 0)
+		{
+			SelectedInstanceFoliageTypes.Add(MeshPair.Key, &MeshPair.Value.Get());
+		}
+	}
+
+	return SelectedInstanceFoliageTypes;
+}
+
+void AInstancedFoliageActor::Destroyed()
+{
+	if (GIsEditor && !GetWorld()->IsGameWorld())
+	{
+		for (auto& MeshPair : FoliageMeshes)
+		{
+			UHierarchicalInstancedStaticMeshComponent* Component = MeshPair.Value->Component;
+
+			if (Component)
+			{
+				Component->ClearInstances();
+				// Save the component's PendingKill flag to restore the component if the delete is undone.
+				Component->SetFlags(RF_Transactional);
+				Component->Modify();
+			}
+		}
+	}
+
+	Super::Destroyed();
+}
+
+void AInstancedFoliageActor::PreEditUndo()
+{
+	Super::PreEditUndo();
+
+	// Remove all delegate as we dont know what the Undo will affect and we will simply readd those still valid afterward
+	for (auto& MeshPair : FoliageMeshes)
+	{
+		FFoliageMeshInfo& MeshInfo = *MeshPair.Value;
+
+		if (MeshPair.Key != nullptr && MeshPair.Key->GetStaticMesh() != nullptr)
+		{
+			MeshPair.Key->GetStaticMesh()->GetOnExtendedBoundsChanged().RemoveAll(&MeshInfo);
+		}
+	}
+}
+
+void AInstancedFoliageActor::PostEditUndo()
+{
+	Super::PostEditUndo();
+
+	FlushRenderingCommands();
+
+	InstanceBaseCache.UpdateInstanceBaseCachedTransforms();
+
+	for (auto& MeshPair : FoliageMeshes)
+	{
+		FFoliageMeshInfo& MeshInfo = *MeshPair.Value;
+
+		if (MeshInfo.Component != nullptr && MeshPair.Key->GetStaticMesh() != nullptr)
+		{
+			MeshPair.Key->GetStaticMesh()->GetOnExtendedBoundsChanged().AddRaw(&MeshInfo, &FFoliageMeshInfo::HandleComponentMeshBoundsChanged);
+		}
+
+		MeshInfo.CheckComponentClass(this, MeshPair.Key);
+		MeshInfo.ReapplyInstancesToComponent();
+
+		// Regenerate instance hash
+		// We regenerate it here instead of saving to transaction buffer to speed up modify operations
+		MeshInfo.InstanceHash->Empty();
+		for (int32 InstanceIdx = 0; InstanceIdx < MeshInfo.Instances.Num(); InstanceIdx++)
+		{
+			MeshInfo.InstanceHash->InsertInstance(MeshInfo.Instances[InstanceIdx].Location, InstanceIdx);
+		}
+	}
+}
+
+bool AInstancedFoliageActor::ShouldExport()
+{
+	// We don't support exporting/importing InstancedFoliageActor itself
+	// Instead foliage instances exported/imported together with components it's painted on
+	return false;
+}
+
+bool AInstancedFoliageActor::ShouldImport(FString* ActorPropString, bool IsMovingLevel)
+{
+	return false;
+}
+
+void AInstancedFoliageActor::ApplySelectionToComponents(bool bApply)
+{
+	for (auto& MeshPair : FoliageMeshes)
+	{
+		FFoliageMeshInfo& MeshInfo = *MeshPair.Value;
+		UHierarchicalInstancedStaticMeshComponent* Component = MeshInfo.Component;
+
+		if (Component && (bApply || Component->SelectedInstances.Num() > 0))
+		{
+			Component->ClearInstanceSelection();
+
+			if (bApply)
+			{
+				for (int32 i : MeshInfo.SelectedIndices)
+				{
+					Component->SelectInstance(true, i, 1);
+				}
+			}
+
+			Component->MarkRenderStateDirty();
+		}
+	}
+}
+
+bool AInstancedFoliageActor::GetSelectionLocation(FVector& OutLocation) const
+{
+	for (const auto& MeshPair : FoliageMeshes)
+	{
+		const FFoliageMeshInfo& MeshInfo = MeshPair.Value.Get();
+		if (MeshInfo.SelectedIndices.Num())
+		{
+			const int32 InstanceIdx = (*MeshInfo.SelectedIndices.CreateConstIterator());
+			OutLocation = MeshInfo.Instances[InstanceIdx].Location;
+			return true;
+		}
+	}
+	return false;
+}
+
+bool AInstancedFoliageActor::HasFoliageAttached(UActorComponent* InComponent)
+{
+	for (TActorIterator<AInstancedFoliageActor> It(InComponent->GetWorld()); It; ++It)
+	{
+		AInstancedFoliageActor* IFA = (*It);
+		if (IFA->InstanceBaseCache.GetInstanceBaseId(InComponent) != FFoliageInstanceBaseCache::InvalidBaseId)
+		{
+			return true;
+		}
+	}
+
+	return false;
+}
+
+
+void AInstancedFoliageActor::MapRebuild()
+{
+	// Map rebuild may have modified the BSP's ModelComponents and thrown the previous ones away.
+	// Most BSP-painted foliage is attached to a Brush's UModelComponent which persist across rebuilds,
+	// but any foliage attached directly to the level BSP's ModelComponents will need to try to find a new base.
+
+	TMap<UFoliageType*, TArray<FFoliageInstance>> NewInstances;
+	TArray<UModelComponent*> RemovedModelComponents;
+	UWorld* World = GetWorld();
+	check(World);
+
+	// For each foliage brush, represented by the mesh/info pair
+	for (auto& MeshPair : FoliageMeshes)
+	{
+		// each target component has some foliage instances
+		FFoliageMeshInfo const& MeshInfo = *MeshPair.Value;
+		UFoliageType* Settings = MeshPair.Key;
+		check(Settings);
+
+		for (auto& ComponentFoliagePair : MeshInfo.ComponentHash)
+		{
+			// BSP components are UModelComponents - they are the only ones we need to change
+			auto BaseComponentPtr = InstanceBaseCache.GetInstanceBasePtr(ComponentFoliagePair.Key);
+			UModelComponent* TargetComponent = Cast<UModelComponent>(BaseComponentPtr.Get());
+
+			// Check if it's part of a brush. We only need to fix up model components that are part of the level BSP.
+			if (TargetComponent && Cast<ABrush>(TargetComponent->GetOuter()) == nullptr)
+			{
+				// Delete its instances later
+				RemovedModelComponents.Add(TargetComponent);
+
+				// We have to test each instance to see if we can migrate it across
+				for (int32 InstanceIdx : ComponentFoliagePair.Value)
+				{
+					// Use a line test against the world. This is not very reliable as we don't know the original trace direction.
+					check(MeshInfo.Instances.IsValidIndex(InstanceIdx));
+					FFoliageInstance const& Instance = MeshInfo.Instances[InstanceIdx];
+
+					FFoliageInstance NewInstance = Instance;
+
+					FTransform InstanceToWorld = Instance.GetInstanceWorldTransform();
+					FVector Down(-FVector::UpVector);
+					FVector Start(InstanceToWorld.TransformPosition(FVector::UpVector));
+					FVector End(InstanceToWorld.TransformPosition(Down));
+
+					FHitResult Result;
+					bool bHit = World->LineTraceSingleByObjectType(Result, Start, End, FCollisionObjectQueryParams(ECC_WorldStatic), FCollisionQueryParams(NAME_None, FCollisionQueryParams::GetUnknownStatId(), true));
+
+					if (bHit && Result.Component.IsValid() && Result.Component->IsA(UModelComponent::StaticClass()))
+					{
+						NewInstance.BaseId = InstanceBaseCache.AddInstanceBaseId(Result.Component.Get());
+						NewInstances.FindOrAdd(Settings).Add(NewInstance);
+					}
+				}
+			}
+		}
+	}
+
+	// Remove all existing & broken instances & component references.
+	for (UModelComponent* Component : RemovedModelComponents)
+	{
+		DeleteInstancesForComponent(Component);
+	}
+
+	// And then finally add our new instances to the correct target components.
+	for (auto& NewInstancePair : NewInstances)
+	{
+		UFoliageType* Settings = NewInstancePair.Key;
+		check(Settings);
+		FFoliageMeshInfo& MeshInfo = *FindOrAddMesh(Settings);
+		for (FFoliageInstance& Instance : NewInstancePair.Value)
+		{
+			MeshInfo.AddInstance(this, Settings, Instance, false);
+		}
+
+		MeshInfo.Component->BuildTreeIfOutdated(true, true);
+	}
+}
+
+#endif // WITH_EDITOR
+
+struct FFoliageMeshInfo_Old
+{
+	TArray<FFoliageInstanceCluster_Deprecated> InstanceClusters;
+	TArray<FFoliageInstance_Deprecated> Instances;
+	UFoliageType_InstancedStaticMesh* Settings; // Type remapped via +ActiveClassRedirects
+};
+FArchive& operator<<(FArchive& Ar, FFoliageMeshInfo_Old& MeshInfo)
+{
+	Ar << MeshInfo.InstanceClusters;
+	Ar << MeshInfo.Instances;
+	Ar << MeshInfo.Settings;
+
+	return Ar;
+}
+
+void AInstancedFoliageActor::Serialize(FArchive& Ar)
+{
+	Super::Serialize(Ar);
+
+	Ar.UsingCustomVersion(FFoliageCustomVersion::GUID);
+
+#if WITH_EDITORONLY_DATA
+	if (!Ar.ArIsFilterEditorOnly && Ar.CustomVer(FFoliageCustomVersion::GUID) >= FFoliageCustomVersion::CrossLevelBase)
+	{
+		Ar << InstanceBaseCache;
+	}
+#endif
+
+	if (Ar.UE4Ver() < VER_UE4_FOLIAGE_SETTINGS_TYPE)
+	{
+#if WITH_EDITORONLY_DATA
+		TMap<UFoliageType*, TUniqueObj<FFoliageMeshInfo_Deprecated>> FoliageMeshesDeprecated;
+		TMap<UStaticMesh*, FFoliageMeshInfo_Old> OldFoliageMeshes;
+		Ar << OldFoliageMeshes;
+		for (auto& OldMeshInfo : OldFoliageMeshes)
+		{
+			FFoliageMeshInfo_Deprecated NewMeshInfo;
+
+			NewMeshInfo.Instances = MoveTemp(OldMeshInfo.Value.Instances);
+
+			UFoliageType_InstancedStaticMesh* FoliageType = OldMeshInfo.Value.Settings;
+			if (FoliageType == nullptr)
+			{
+				// If the Settings object was null, eg the user forgot to save their settings asset, create a new one.
+				FoliageType = NewObject<UFoliageType_InstancedStaticMesh>(this);
+			}
+
+			if (FoliageType->Mesh == nullptr)
+			{
+				FoliageType->Modify();
+				FoliageType->Mesh = OldMeshInfo.Key;
+			}
+			else if (FoliageType->Mesh != OldMeshInfo.Key)
+			{
+				// If mesh doesn't match (two meshes sharing the same settings object?) then we need to duplicate as that is no longer supported
+				FoliageType = (UFoliageType_InstancedStaticMesh*)StaticDuplicateObject(FoliageType, this, NAME_None, RF_AllFlags & ~(RF_Standalone | RF_Public));
+				FoliageType->Mesh = OldMeshInfo.Key;
+			}
+			NewMeshInfo.FoliageTypeUpdateGuid = FoliageType->UpdateGuid;
+			FoliageMeshes_Deprecated.Add(FoliageType, TUniqueObj<FFoliageMeshInfo_Deprecated>(MoveTemp(NewMeshInfo)));
+		}
+#endif//WITH_EDITORONLY_DATA
+	}
+	else
+	{
+		if (Ar.CustomVer(FFoliageCustomVersion::GUID) < FFoliageCustomVersion::CrossLevelBase)
+		{
+#if WITH_EDITORONLY_DATA
+			Ar << FoliageMeshes_Deprecated;
+#endif
+		}
+		else
+		{
+			Ar << FoliageMeshes;
+		}
+	}
+
+	// Clean up any old cluster components and convert to hierarchical instanced foliage.
+	if (Ar.CustomVer(FFoliageCustomVersion::GUID) < FFoliageCustomVersion::FoliageUsingHierarchicalISMC)
+	{
+		for (UActorComponent* Component : GetComponents())
+		{
+			if (Cast<UInstancedStaticMeshComponent>(Component))
+			{
+				Component->bAutoRegister = false;
+			}
+		}
+	}
+}
+
+#if WITH_EDITOR
+void AInstancedFoliageActor::PostInitProperties()
+{
+	Super::PostInitProperties();
+
+	if (!IsTemplate())
+	{
+		GEngine->OnActorMoved().Remove(OnLevelActorMovedDelegateHandle);
+		OnLevelActorMovedDelegateHandle = GEngine->OnActorMoved().AddUObject(this, &AInstancedFoliageActor::OnLevelActorMoved);
+
+		GEngine->OnLevelActorDeleted().Remove(OnLevelActorDeletedDelegateHandle);
+		OnLevelActorDeletedDelegateHandle = GEngine->OnLevelActorDeleted().AddUObject(this, &AInstancedFoliageActor::OnLevelActorDeleted);
+
+		FWorldDelegates::PostApplyLevelOffset.Remove(OnPostApplyLevelOffsetDelegateHandle);
+		OnPostApplyLevelOffsetDelegateHandle = FWorldDelegates::PostApplyLevelOffset.AddUObject(this, &AInstancedFoliageActor::OnPostApplyLevelOffset);
+	}
+}
+
+void AInstancedFoliageActor::BeginDestroy()
+{
+	Super::BeginDestroy();
+
+	if (!IsTemplate())
+	{
+		GEngine->OnActorMoved().Remove(OnLevelActorMovedDelegateHandle);
+		GEngine->OnLevelActorDeleted().Remove(OnLevelActorDeletedDelegateHandle);
+		FWorldDelegates::PostApplyLevelOffset.Remove(OnPostApplyLevelOffsetDelegateHandle);
+	}
+}
+#endif
+
+void AInstancedFoliageActor::PostLoad()
+{
+	Super::PostLoad();
+
+	ULevel* OwningLevel = GetLevel();
+	if (!OwningLevel->InstancedFoliageActor.IsValid())
+	{
+		OwningLevel->InstancedFoliageActor = this;
+	}
+	else
+	{
+		FFormatNamedArguments Arguments;
+		Arguments.Add(TEXT("Level"), FText::FromString(*OwningLevel->GetOutermost()->GetName()));
+		FMessageLog("MapCheck").Warning()
+			->AddToken(FUObjectToken::Create(this))
+			->AddToken(FTextToken::Create(FText::Format(LOCTEXT("MapCheck_DuplicateInstancedFoliageActor", "Level {Level} has an unexpected duplicate Instanced Foliage Actor."), Arguments)))
+#if WITH_EDITOR
+			->AddToken(FActionToken::Create(LOCTEXT("MapCheck_FixDuplicateInstancedFoliageActor", "Fix"),
+				LOCTEXT("MapCheck_FixDuplicateInstancedFoliageActor_Desc", "Click to consolidate foliage into the main foliage actor."),
+				FOnActionTokenExecuted::CreateUObject(OwningLevel->InstancedFoliageActor.Get(), &AInstancedFoliageActor::RepairDuplicateIFA, this), true))
+#endif// WITH_EDITOR
+			;
+		FMessageLog("MapCheck").Open(EMessageSeverity::Warning);
+	}
+
+#if WITH_EDITOR
+	if (GIsEditor)
+	{
+		if (GetLinkerCustomVersion(FFoliageCustomVersion::GUID) < FFoliageCustomVersion::CrossLevelBase)
+		{
+			ConvertDeprecatedFoliageMeshes(this, FoliageMeshes_Deprecated, FoliageMeshes);
+			FoliageMeshes_Deprecated.Empty();
+		}
+
+		{
+			bool bContainsNull = FoliageMeshes.Remove(nullptr) > 0;
+			if (bContainsNull)
+			{
+				FMessageLog("MapCheck").Warning()
+					->AddToken(FUObjectToken::Create(this))
+					->AddToken(FTextToken::Create(LOCTEXT("MapCheck_Message_FoliageMissingStaticMesh", "Foliage instances for a missing static mesh have been removed.")))
+					->AddToken(FMapErrorToken::Create(FMapErrors::FoliageMissingStaticMesh));
+				while (bContainsNull)
+				{
+					bContainsNull = FoliageMeshes.Remove(nullptr) > 0;
+				}
+			}
+		}
+
+		TArray<UFoliageType*> FoliageTypeToRemove;
+
+		for (auto& MeshPair : FoliageMeshes)
+		{
+			// Find the per-mesh info matching the mesh.
+			FFoliageMeshInfo& MeshInfo = *MeshPair.Value;
+			UFoliageType* FoliageType = MeshPair.Key;
+
+			// Make sure the mesh has been PostLoaded as if not it can be considered invalid resulting in a bad HISMC tree
+			UStaticMesh* StaticMesh = FoliageType->GetStaticMesh();
+			if (StaticMesh)
+			{
+				StaticMesh->ConditionalPostLoad();
+			}
+
+			if (MeshInfo.Instances.Num() && MeshInfo.Component == nullptr)
+			{
+				FFormatNamedArguments Arguments;
+				if (StaticMesh)
+				{
+					Arguments.Add(TEXT("MeshName"), FText::FromString(StaticMesh->GetName()));
+				}
+				else
+				{
+					Arguments.Add(TEXT("MeshName"), FText::FromString(TEXT("None")));
+				}
+
+				FMessageLog("MapCheck").Warning()
+					->AddToken(FUObjectToken::Create(this))
+					->AddToken(FTextToken::Create(FText::Format(LOCTEXT("MapCheck_Message_FoliageMissingComponent", "Foliage in this map is missing a component for static mesh {MeshName}. This has been repaired."), Arguments)))
+					->AddToken(FMapErrorToken::Create(FMapErrors::FoliageMissingClusterComponent));
+
+				MeshInfo.ReallocateClusters(this, MeshPair.Key);
+			}
+
+			// Update the hash.
+			MeshInfo.ComponentHash.Empty();
+			MeshInfo.InstanceHash->Empty();
+			for (int32 InstanceIdx = 0; InstanceIdx < MeshInfo.Instances.Num(); InstanceIdx++)
+			{
+				MeshInfo.AddToBaseHash(InstanceIdx);
+				MeshInfo.InstanceHash->InsertInstance(MeshInfo.Instances[InstanceIdx].Location, InstanceIdx);
+			}
+
+			// Convert to Hierarchical foliage
+			if (GetLinkerCustomVersion(FFoliageCustomVersion::GUID) < FFoliageCustomVersion::FoliageUsingHierarchicalISMC)
+			{
+				MeshInfo.ReallocateClusters(this, MeshPair.Key);
+			}
+
+			if (GetLinkerCustomVersion(FFoliageCustomVersion::GUID) < FFoliageCustomVersion::HierarchicalISMCNonTransactional)
+			{
+				if (MeshInfo.Component)
+				{
+					MeshInfo.Component->ClearFlags(RF_Transactional);
+				}
+			}
+
+			// Clean up case where embeded instances had their static mesh deleted
+			if (FoliageType->IsNotAssetOrBlueprint() && StaticMesh == nullptr)
+			{
+				// We can't remove them here as we are within the loop itself so clean up after
+				FoliageTypeToRemove.Add(FoliageType);
+				
+				continue;
+			}
+
+			// Upgrade foliage component
+			if (GetLinkerCustomVersion(FFoliageCustomVersion::GUID) < FFoliageCustomVersion::FoliageUsingFoliageISMC)
+			{
+				MeshInfo.CheckComponentClass(this, FoliageType);
+			}
+
+			// Update foliage component settings if the foliage settings object was changed while the level was not loaded.
+			if (MeshInfo.FoliageTypeUpdateGuid != FoliageType->UpdateGuid)
+			{
+				if (MeshInfo.FoliageTypeUpdateGuid.IsValid())
+				{
+					MeshInfo.CheckComponentClass(this, FoliageType);
+					MeshInfo.UpdateComponentSettings(FoliageType);
+				}
+				MeshInfo.FoliageTypeUpdateGuid = FoliageType->UpdateGuid;
+			}
+		}
+
+#if WITH_EDITORONLY_DATA
+		if (GetLinkerCustomVersion(FFortniteMainBranchObjectVersion::GUID) < FFortniteMainBranchObjectVersion::FoliageLazyObjPtrToSoftObjPtr)
+		{
+			for (auto Iter = InstanceBaseCache.InstanceBaseMap.CreateIterator(); Iter; ++Iter)
+			{
+				TPair<FFoliageInstanceBaseId, FFoliageInstanceBaseInfo>& Pair = *Iter;
+				FFoliageInstanceBaseInfo& BaseInfo = Pair.Value;
+				UActorComponent* Component = BaseInfo.BasePtr_DEPRECATED.Get();
+				BaseInfo.BasePtr_DEPRECATED.Reset();
+
+				if (Component != nullptr)
+				{
+					BaseInfo.BasePtr = Component;
+
+					if (!InstanceBaseCache.InstanceBaseInvMap.Contains(BaseInfo.BasePtr))
+					{
+						InstanceBaseCache.InstanceBaseInvMap.Add(BaseInfo.BasePtr, Pair.Key);
+					}
+				}
+				else
+				{
+					Iter.RemoveCurrent();
+
+					const FFoliageInstanceBasePtr* BaseInfoPtr = InstanceBaseCache.InstanceBaseInvMap.FindKey(Pair.Key);
+
+					if (BaseInfoPtr != nullptr && BaseInfoPtr->Get() == nullptr)
+					{
+						InstanceBaseCache.InstanceBaseInvMap.Remove(*BaseInfoPtr);
+					}
+				}
+			}
+
+			InstanceBaseCache.InstanceBaseMap.Compact();
+			InstanceBaseCache.InstanceBaseInvMap.Compact();
+
+			for (auto& Pair : InstanceBaseCache.InstanceBaseLevelMap_DEPRECATED)
+			{
+				TArray<FFoliageInstanceBasePtr_DEPRECATED>& BaseInfo_DEPRECATED = Pair.Value;
+				TArray<FFoliageInstanceBasePtr> BaseInfo;
+
+				for (FFoliageInstanceBasePtr_DEPRECATED& BasePtr_DEPRECATED : BaseInfo_DEPRECATED)
+				{
+					UActorComponent* Component = BasePtr_DEPRECATED.Get();
+					BasePtr_DEPRECATED.Reset();
+
+					if (Component != nullptr)
+					{
+						BaseInfo.Add(Component);
+					}
+				}
+
+				InstanceBaseCache.InstanceBaseLevelMap.Add(Pair.Key, BaseInfo);
+			}
+
+			InstanceBaseCache.InstanceBaseLevelMap_DEPRECATED.Empty();
+		}
+
+		// Clean up dead cross-level references
+		FFoliageInstanceBaseCache::CompactInstanceBaseCache(this);
+#endif
+
+		// Clean up invalid foliage type
+		for (UFoliageType* FoliageType : FoliageTypeToRemove)
+		{
+			OnFoliageTypeMeshChangedEvent.Broadcast(FoliageType);
+			RemoveFoliageType(&FoliageType, 1);
+		}
+	}
+
+#endif// WITH_EDITOR
+
+	if (!GIsEditor && CVarFoliageDiscardDataOnLoad.GetValueOnGameThread())
+	{
+		for (auto& MeshPair : FoliageMeshes)
+		{
+			if (MeshPair.Value->Component != nullptr && (!MeshPair.Key || MeshPair.Key->bEnableDensityScaling))
+			{
+				MeshPair.Value->Component->ConditionalPostLoad();
+				MeshPair.Value->Component->DestroyComponent();
+			}
+
+			MeshPair.Value = FFoliageMeshInfo();
+		}
+	}
+}
+
+#if WITH_EDITOR
+
+void AInstancedFoliageActor::RepairDuplicateIFA(AInstancedFoliageActor* DuplicateIFA)
+{
+	for (auto& MeshPair : DuplicateIFA->FoliageMeshes)
+	{
+		UFoliageType* DupeFoliageType = MeshPair.Key;
+		FFoliageMeshInfo& DupeMeshInfo = *MeshPair.Value;
+
+		// Get foliage type compatible with target IFA
+		FFoliageMeshInfo* TargetMeshInfo = nullptr;
+		UFoliageType* TargetFoliageType = AddFoliageType(DupeFoliageType, &TargetMeshInfo);
+
+		// Copy the instances
+		for (FFoliageInstance& Instance : DupeMeshInfo.Instances)
+		{
+			if ((Instance.Flags & FOLIAGE_InstanceDeleted) == 0)
+			{
+				TargetMeshInfo->AddInstance(this, TargetFoliageType, Instance, false);
+			}
+		}
+
+		TargetMeshInfo->Component->BuildTreeIfOutdated(true, true);
+	}
+
+	GetWorld()->DestroyActor(DuplicateIFA);
+}
+
+void AInstancedFoliageActor::NotifyFoliageTypeChanged(UFoliageType* FoliageType, bool bMeshChanged)
+{
+	FFoliageMeshInfo* TypeInfo = FindMesh(FoliageType);
+
+	if (TypeInfo)
+	{
+		TypeInfo->CheckComponentClass(this, FoliageType);
+		TypeInfo->UpdateComponentSettings(FoliageType);
+
+		if (bMeshChanged)
+		{
+			// If the type's mesh has changed, the UI needs to be notified so it can update thumbnails accordingly
+			OnFoliageTypeMeshChangedEvent.Broadcast(FoliageType);
+
+			// Change bounds delegate bindings
+			if (TypeInfo->Component != nullptr && TypeInfo->Component->GetStaticMesh() != nullptr && FoliageType->GetStaticMesh() != nullptr)
+			{
+				TypeInfo->Component->GetStaticMesh()->GetOnExtendedBoundsChanged().AddRaw(TypeInfo, &FFoliageMeshInfo::HandleComponentMeshBoundsChanged);
+
+				// Mesh changed, so we must update the occlusion tree
+				TypeInfo->Component->BuildTreeIfOutdated(true, false);
+			}
+
+			if (FoliageType->IsNotAssetOrBlueprint() && FoliageType->GetStaticMesh() == nullptr) //If the mesh has been deleted and we're a per foliage actor instance we must remove all instances of the mesh
+			{
+				RemoveFoliageType(&FoliageType, 1);
+			}
+		}
+	}
+}
+
+void AInstancedFoliageActor::NotifyFoliageTypeWillChange(UFoliageType* FoliageType, bool bMeshChanged)
+{
+	if (bMeshChanged)
+	{
+		FFoliageMeshInfo* TypeInfo = FindMesh(FoliageType);
+
+		// Change bounds delegate bindings
+		if (TypeInfo)
+		{
+			if (TypeInfo->Component != nullptr && TypeInfo->Component->GetStaticMesh() != nullptr)
+			{
+				TypeInfo->Component->GetStaticMesh()->GetOnExtendedBoundsChanged().RemoveAll(TypeInfo);
+			}
+		}
+	}
+}
+
+void AInstancedFoliageActor::OnLevelActorMoved(AActor* InActor)
+{
+	UWorld* InWorld = InActor->GetWorld();
+
+	if (!InWorld || !InWorld->IsGameWorld())
+	{
+		for (UActorComponent* Component : InActor->GetComponents())
+		{
+			if (Component)
+			{
+				MoveInstancesForMovedComponent(Component);
+			}
+		}
+	}
+}
+
+void AInstancedFoliageActor::OnLevelActorDeleted(AActor* InActor)
+{
+	UWorld* InWorld = InActor->GetWorld();
+
+	if (!InWorld || !InWorld->IsGameWorld())
+	{
+		for (UActorComponent* Component : InActor->GetComponents())
+		{
+			if (Component)
+			{
+				DeleteInstancesForComponent(Component);
+			}
+		}
+	}
+}
+
+void AInstancedFoliageActor::OnPostApplyLevelOffset(ULevel* InLevel, UWorld* InWorld, const FVector& InOffset, bool bWorldShift)
+{
+	ULevel* OwningLevel = GetLevel();
+	if (InLevel != OwningLevel) // TODO: cross-level foliage 
+	{
+		return;
+	}
+
+	if (GIsEditor && InWorld && !InWorld->IsGameWorld())
+	{
+		for (auto& MeshPair : FoliageMeshes)
+		{
+			FFoliageMeshInfo& MeshInfo = *MeshPair.Value;
+
+			InstanceBaseCache.UpdateInstanceBaseCachedTransforms();
+
+			MeshInfo.InstanceHash->Empty();
+			for (int32 InstanceIdx = 0; InstanceIdx < MeshInfo.Instances.Num(); InstanceIdx++)
+			{
+				FFoliageInstance& Instance = MeshInfo.Instances[InstanceIdx];
+				Instance.Location += InOffset;
+				// Rehash instance location
+				MeshInfo.InstanceHash->InsertInstance(Instance.Location, InstanceIdx);
+			}
+		}
+	}
+}
+
+
+void AInstancedFoliageActor::CleanupDeletedFoliageType()
+{
+	for (auto& MeshPair : FoliageMeshes)
+	{
+		if (MeshPair.Key == nullptr)
+		{
+			FFoliageMeshInfo& MeshInfo = *MeshPair.Value;
+			TArray<int32> InstancesToRemove;
+			for (int32 InstanceIdx = 0; InstanceIdx < MeshInfo.Instances.Num(); InstanceIdx++)
+			{
+				InstancesToRemove.Add(InstanceIdx);
+			}
+
+			if (InstancesToRemove.Num())
+			{
+				MeshInfo.RemoveInstances(this, InstancesToRemove, true);
+			}
+		}
+
+	}
+
+	while (FoliageMeshes.Remove(nullptr)) {}	//remove entries from the map
+
+}
+
+
+#endif
+//
+// Serialize all our UObjects for RTGC 
+//
+void AInstancedFoliageActor::AddReferencedObjects(UObject* InThis, FReferenceCollector& Collector)
+{
+	AInstancedFoliageActor* This = CastChecked<AInstancedFoliageActor>(InThis);
+
+	for (auto& MeshPair : This->FoliageMeshes)
+	{
+		Collector.AddReferencedObject(MeshPair.Key, This);
+		FFoliageMeshInfo& MeshInfo = *MeshPair.Value;
+
+		if (MeshInfo.Component)
+		{
+			Collector.AddReferencedObject(MeshInfo.Component, This);
+		}
+	}
+
+	Super::AddReferencedObjects(This, Collector);
+}
+
+#if WITH_EDITOR
+bool AInstancedFoliageActor::FoliageTrace(const UWorld* InWorld, FHitResult& OutHit, const FDesiredFoliageInstance& DesiredInstance, FName InTraceTag, bool InbReturnFaceIndex, const FFoliageTraceFilterFunc& FilterFunc)
+{
+	SCOPE_CYCLE_COUNTER(STAT_FoliageTrace);
+
+	FCollisionQueryParams QueryParams(InTraceTag, SCENE_QUERY_STAT_ONLY(IFA_FoliageTrace), true);
+	QueryParams.bReturnFaceIndex = InbReturnFaceIndex;
+
+	//It's possible that with the radius of the shape we will end up with an initial overlap which would place the instance at the top of the procedural volume.
+	//Moving the start trace back a bit will fix this, but it introduces the potential for spawning instances a bit above the volume. This second issue is already somewhat broken because of how sweeps work so it's not too bad, also this is a less common case.
+	//The proper fix would be to do something like EncroachmentCheck where we first do a sweep, then we fix it up if it's overlapping, then check the filters. This is more expensive and error prone so for now we just move the trace up a bit.
+	const FVector Dir = (DesiredInstance.EndTrace - DesiredInstance.StartTrace).GetSafeNormal();
+	const FVector StartTrace = DesiredInstance.StartTrace - (Dir * DesiredInstance.TraceRadius);
+
+	TArray<FHitResult> Hits;
+	FCollisionShape SphereShape;
+	SphereShape.SetSphere(DesiredInstance.TraceRadius);
+	InWorld->SweepMultiByObjectType(Hits, StartTrace, DesiredInstance.EndTrace, FQuat::Identity, FCollisionObjectQueryParams(ECC_WorldStatic), SphereShape, QueryParams);
+
+	for (const FHitResult& Hit : Hits)
+	{
+		const AActor* HitActor = Hit.GetActor();
+
+		// don't place procedural foliage inside an AProceduralFoliageBlockingVolume
+		// this test is first because two of the tests below would otherwise cause the trace to ignore AProceduralFoliageBlockingVolume
+		if (DesiredInstance.PlacementMode == EFoliagePlacementMode::Procedural)
+		{
+			if (const AProceduralFoliageBlockingVolume* ProceduralFoliageBlockingVolume = Cast<AProceduralFoliageBlockingVolume>(HitActor))
+			{
+				const AProceduralFoliageVolume* ProceduralFoliageVolume = ProceduralFoliageBlockingVolume->ProceduralFoliageVolume;
+				if (ProceduralFoliageVolume == nullptr || ProceduralFoliageVolume->ProceduralComponent == nullptr || ProceduralFoliageVolume->ProceduralComponent->GetProceduralGuid() == DesiredInstance.ProceduralGuid)
+				{
+					return false;
+				}
+			}
+			else if (HitActor && HitActor->IsA<AProceduralFoliageVolume>()) //we never want to collide with our spawning volume
+			{
+				continue;
+			}
+		}
+
+		const UPrimitiveComponent* HitComponent = Hit.GetComponent();
+		check(HitComponent);
+
+		// In the editor traces can hit "No Collision" type actors, so ugh. (ignore these)
+		if (!HitComponent->IsQueryCollisionEnabled() || HitComponent->GetCollisionResponseToChannel(ECC_WorldStatic) != ECR_Block)
+		{
+			continue;
+		}
+
+		// Don't place foliage on invisible walls / triggers / volumes
+		if (HitComponent->IsA<UBrushComponent>())
+		{
+			continue;
+		}
+
+		// Don't place foliage on itself
+		if (const AInstancedFoliageActor* FoliageActor = Cast<AInstancedFoliageActor>(HitActor))
+		{
+			if (const FFoliageMeshInfo* FoundMeshInfo = FoliageActor->FindMesh(DesiredInstance.FoliageType))
+			{
+				if (FoundMeshInfo->Component == HitComponent)
+				{
+					continue;
+				}
+			}
+		}
+
+		if (FilterFunc && FilterFunc(HitComponent) == false)
+		{
+			// supplied filter does not like this component, so keep iterating
+			continue;
+		}
+
+		bool bInsideProceduralVolumeOrArentUsingOne = true;
+		if (DesiredInstance.PlacementMode == EFoliagePlacementMode::Procedural && DesiredInstance.ProceduralVolumeBodyInstance)
+		{
+			// We have a procedural volume, so lets make sure we are inside it.
+			bInsideProceduralVolumeOrArentUsingOne = DesiredInstance.ProceduralVolumeBodyInstance->OverlapTest(Hit.ImpactPoint, FQuat::Identity, FCollisionShape::MakeSphere(1.f));	//make sphere of 1cm radius to test if we're in the procedural volume
+		}
+
+		OutHit = Hit;
+
+		// When placing foliage on other foliage, we need to return the base component of the other foliage, not the foliage component, so that it moves correctly
+		if (const AInstancedFoliageActor* FoliageActor = Cast<AInstancedFoliageActor>(HitActor))
+		{
+			for (auto& MeshPair : FoliageActor->FoliageMeshes)
+			{
+				const FFoliageMeshInfo& MeshInfo = *MeshPair.Value;
+				if (MeshInfo.Component == HitComponent)
+				{
+					OutHit.Component = CastChecked<UPrimitiveComponent>(FoliageActor->InstanceBaseCache.GetInstanceBasePtr(MeshInfo.Instances[Hit.Item].BaseId).Get(), ECastCheckedType::NullAllowed);
+					break;
+				}
+			}
+		}
+
+		return bInsideProceduralVolumeOrArentUsingOne;
+	}
+
+	return false;
+}
+
+bool AInstancedFoliageActor::CheckCollisionWithWorld(const UWorld* InWorld, const UFoliageType* Settings, const FFoliageInstance& Inst, const FVector& HitNormal, const FVector& HitLocation, UPrimitiveComponent* HitComponent)
+{
+	if (!Settings->CollisionWithWorld)
+	{
+		return true;
+	}
+
+	FTransform OriginalTransform = Inst.GetInstanceWorldTransform();
+	OriginalTransform.SetRotation(FQuat::Identity);
+
+	FMatrix InstTransformNoRotation = OriginalTransform.ToMatrixWithScale();
+	OriginalTransform = Inst.GetInstanceWorldTransform();
+
+	// Check for overhanging ledge
+	const int32 SamplePositionCount = 4;
+	{
+		FVector LocalSamplePos[SamplePositionCount] = {
+			FVector(Settings->LowBoundOriginRadius.Z, 0, 0),
+			FVector(-Settings->LowBoundOriginRadius.Z, 0, 0),
+			FVector(0, Settings->LowBoundOriginRadius.Z, 0),
+			FVector(0, -Settings->LowBoundOriginRadius.Z, 0)
+		};
+
+		for (uint32 i = 0; i < SamplePositionCount; ++i)
+		{
+			FVector SamplePos = InstTransformNoRotation.TransformPosition(Settings->LowBoundOriginRadius + LocalSamplePos[i]);
+			float WorldRadius = (Settings->LowBoundOriginRadius.Z + Settings->LowBoundOriginRadius.Z)*FMath::Max(Inst.DrawScale3D.X, Inst.DrawScale3D.Y);
+			FVector NormalVector = Settings->AlignToNormal ? HitNormal : OriginalTransform.GetRotation().GetUpVector();
+
+			//::DrawDebugSphere(InWorld, SamplePos, 10, 6, FColor::Red, true, 30.0f);
+			//::DrawDebugSphere(InWorld, SamplePos - NormalVector*WorldRadius, 10, 6, FColor::Orange, true, 30.0f);
+			//::DrawDebugDirectionalArrow(InWorld, SamplePos, SamplePos - NormalVector*WorldRadius, 10.0f, FColor::Red, true, 30.0f);
+
+			FHitResult Hit;
+			if (AInstancedFoliageActor::FoliageTrace(InWorld, Hit, FDesiredFoliageInstance(SamplePos, SamplePos - NormalVector*WorldRadius)))
+			{
+				FVector LocalHit = OriginalTransform.InverseTransformPosition(Hit.Location);
+				
+				if (LocalHit.Z - Inst.ZOffset < Settings->LowBoundOriginRadius.Z && Hit.Component.Get() == HitComponent)
+				{
+					//::DrawDebugSphere(InWorld, Hit.Location, 6, 6, FColor::Green, true, 30.0f);
+					continue;
+				}
+			}
+
+			//::DrawDebugSphere(InWorld, SamplePos, 6, 6, FColor::Cyan, true, 30.0f);
+
+			return false;
+		}
+	}
+
+	FBoxSphereBounds LocalBound(Settings->MeshBounds.GetBox());
+	FBoxSphereBounds WorldBound = LocalBound.TransformBy(OriginalTransform);
+
+	static FName NAME_FoliageCollisionWithWorld = FName(TEXT("FoliageCollisionWithWorld"));
+	if (InWorld->OverlapBlockingTestByChannel(WorldBound.Origin, FQuat(Inst.Rotation), ECC_WorldStatic, FCollisionShape::MakeBox(LocalBound.BoxExtent * Inst.DrawScale3D * Settings->CollisionScale), FCollisionQueryParams(NAME_FoliageCollisionWithWorld, false, HitComponent != nullptr ? HitComponent->GetOwner() : nullptr)))
+	{
+		return false;
+	}
+
+	//::DrawDebugBox(InWorld, WorldBound.Origin, LocalBound.BoxExtent * Inst.DrawScale3D * Settings->CollisionScale, FQuat(Inst.Rotation), FColor::Red, true, 30.f);
+
+	return true;
+}
+
+FPotentialInstance::FPotentialInstance(FVector InHitLocation, FVector InHitNormal, UPrimitiveComponent* InHitComponent, float InHitWeight, const FDesiredFoliageInstance& InDesiredInstance)
+	: HitLocation(InHitLocation)
+	, HitNormal(InHitNormal)
+	, HitComponent(InHitComponent)
+	, HitWeight(InHitWeight)
+	, DesiredInstance(InDesiredInstance)
+{
+}
+
+bool FPotentialInstance::PlaceInstance(const UWorld* InWorld, const UFoliageType* Settings, FFoliageInstance& Inst, bool bSkipCollision)
+{
+	if (DesiredInstance.PlacementMode != EFoliagePlacementMode::Procedural)
+	{
+		Inst.DrawScale3D = Settings->GetRandomScale();
+	}
+	else
+	{
+		//Procedural foliage uses age to get the scale
+		Inst.DrawScale3D = FVector(Settings->GetScaleForAge(DesiredInstance.Age));
+	}
+
+	Inst.ZOffset = Settings->ZOffset.Interpolate(FMath::FRand());
+
+	Inst.Location = HitLocation;
+
+	if (DesiredInstance.PlacementMode != EFoliagePlacementMode::Procedural)
+	{
+		// Random yaw and optional random pitch up to the maximum
+		Inst.Rotation = FRotator(FMath::FRand() * Settings->RandomPitchAngle, 0.f, 0.f);
+
+		if (Settings->RandomYaw)
+		{
+			Inst.Rotation.Yaw = FMath::FRand() * 360.f;
+		}
+		else
+		{
+			Inst.Flags |= FOLIAGE_NoRandomYaw;
+		}
+	}
+	else
+	{
+		Inst.Rotation = DesiredInstance.Rotation.Rotator();
+		Inst.Flags |= FOLIAGE_NoRandomYaw;
+	}
+
+
+	if (Settings->AlignToNormal)
+	{
+		Inst.AlignToNormal(HitNormal, Settings->AlignMaxAngle);
+	}
+
+	// Apply the Z offset in local space
+	if (FMath::Abs(Inst.ZOffset) > KINDA_SMALL_NUMBER)
+	{
+		Inst.Location = Inst.GetInstanceWorldTransform().TransformPosition(FVector(0, 0, Inst.ZOffset));
+	}
+
+	UModelComponent* ModelComponent = Cast<UModelComponent>(HitComponent);
+	if (ModelComponent)
+	{
+		ABrush* BrushActor = ModelComponent->GetModel()->FindBrush(HitLocation);
+		if (BrushActor)
+		{
+			HitComponent = BrushActor->GetBrushComponent();
+		}
+	}
+
+	return bSkipCollision || AInstancedFoliageActor::CheckCollisionWithWorld(InWorld, Settings, Inst, HitNormal, HitLocation, HitComponent);
+}
+#endif
+
+float AInstancedFoliageActor::InternalTakeRadialDamage(float Damage, struct FRadialDamageEvent const& RadialDamageEvent, class AController* EventInstigator, AActor* DamageCauser)
+{
+	// Radial damage scaling needs to be applied per instance so we don't do anything here
+	return Damage;
+}
+
+UFoliageInstancedStaticMeshComponent::UFoliageInstancedStaticMeshComponent(const FObjectInitializer& ObjectInitializer)
+	: Super(ObjectInitializer)
+{
+#if WITH_EDITORONLY_DATA
+	bEnableAutoLODGeneration = false;
+#endif
+}
+
+void UFoliageInstancedStaticMeshComponent::ReceiveComponentDamage(float DamageAmount, FDamageEvent const& DamageEvent, AController* EventInstigator, AActor* DamageCauser)
+{
+	Super::ReceiveComponentDamage(DamageAmount, DamageEvent, EventInstigator, DamageCauser);
+
+	if (DamageAmount != 0.f)
+	{
+		UDamageType const* const DamageTypeCDO = DamageEvent.DamageTypeClass ? DamageEvent.DamageTypeClass->GetDefaultObject<UDamageType>() : GetDefault<UDamageType>();
+		if (DamageEvent.IsOfType(FPointDamageEvent::ClassID))
+		{
+			// Point damage event, hit a single instance.
+			FPointDamageEvent* const PointDamageEvent = (FPointDamageEvent*)&DamageEvent;
+			if (PerInstanceSMData.IsValidIndex(PointDamageEvent->HitInfo.Item))
+			{
+				OnInstanceTakePointDamage.Broadcast(PointDamageEvent->HitInfo.Item, DamageAmount, EventInstigator, PointDamageEvent->HitInfo.ImpactPoint, PointDamageEvent->ShotDirection, DamageTypeCDO, DamageCauser);
+			}
+		}
+		else if (DamageEvent.IsOfType(FRadialDamageEvent::ClassID))
+		{
+			// Radial damage event, find which instances it hit and notify
+			FRadialDamageEvent* const RadialDamageEvent = (FRadialDamageEvent*)&DamageEvent;
+
+			float MaxRadius = RadialDamageEvent->Params.GetMaxRadius();
+			TArray<int32> Instances = GetInstancesOverlappingSphere(RadialDamageEvent->Origin, MaxRadius, true);
+
+			if (Instances.Num())
+			{
+				FVector LocalOrigin = GetComponentToWorld().Inverse().TransformPosition(RadialDamageEvent->Origin);
+				float Scale = GetComponentScale().X; // assume component (not instances) is uniformly scaled
+
+				TArray<float> Damages;
+				Damages.Empty(Instances.Num());
+
+				for (int32 InstanceIndex : Instances)
+				{
+					// Find distance in local space and then scale; quicker than transforming each instance to world space.
+					float DistanceFromOrigin = (PerInstanceSMData[InstanceIndex].Transform.GetOrigin() - LocalOrigin).Size() * Scale;
+					Damages.Add(RadialDamageEvent->Params.GetDamageScale(DistanceFromOrigin));
+				}
+
+				OnInstanceTakeRadialDamage.Broadcast(Instances, Damages, EventInstigator, RadialDamageEvent->Origin, MaxRadius, DamageTypeCDO, DamageCauser);
+			}
+		}
+	}
+}
+
+#if WITH_EDITOR
+
+uint64 UFoliageInstancedStaticMeshComponent::GetHiddenEditorViews() const
+{
+	return FoliageHiddenEditorViews;
+}
+
+#endif// WITH_EDITOR
+
+#undef LOCTEXT_NAMESPACE
