Fix Build. Submitted to fast.

#rnx

diff --git a/Engine/Source/Runtime/Foliage/Private/InstancedFoliageActor.cpp b/Engine/Source/Runtime/Foliage/Private/InstancedFoliageActor.cpp
--- a/Engine/Source/Runtime/Foliage/Private/InstancedFoliageActor.cpp
+++ b/Engine/Source/Runtime/Foliage/Private/InstancedFoliageActor.cpp
@@ -1,9 +1,10 @@
 // Copyright 1998-2019 Epic Games, Inc. All Rights Reserved.
 
 /*=============================================================================
-InstancedFoliage.cpp: Instanced foliage implementation.
+InstancedFoliageActor.cpp: Instanced foliage implementation.
 =============================================================================*/
 
+#include "InstancedFoliageActor.h"
 #include "InstancedFoliage.h"
 #include "Templates/SubclassOf.h"
 #include "HAL/IConsoleManager.h"
