New Foliage Type: Blueprint Actor
- Allows using the Foliage Editing mode to paint Actors instead of StaticMesh

diff --git a/Engine/Source/Developer/AssetTools/Private/AssetTools.cpp b/Engine/Source/Developer/AssetTools/Private/AssetTools.cpp
--- a/Engine/Source/Developer/AssetTools/Private/AssetTools.cpp
+++ b/Engine/Source/Developer/AssetTools/Private/AssetTools.cpp
@@ -65,6 +65,7 @@
 #include "AssetTypeActions/AssetTypeActions_ForceFeedbackEffect.h"
 #include "AssetTypeActions/AssetTypeActions_HLODProxy.h"
 #include "AssetTypeActions/AssetTypeActions_SubsurfaceProfile.h"
+#include "AssetTypeActions/AssetTypeActions_ActorFoliageSettings.h"
 #include "AssetTypeActions/AssetTypeActions_InstancedFoliageSettings.h"
 #include "AssetTypeActions/AssetTypeActions_InterpData.h"
 #include "AssetTypeActions/AssetTypeActions_LandscapeLayer.h"
@@ -158,6 +159,7 @@ UAssetToolsImpl::UAssetToolsImpl(const FObjectInitializer& ObjectInitializer)
 	AllocatedCategoryBits.Add(TEXT("_BuiltIn_8"), FAdvancedAssetCategory(EAssetTypeCategories::Media, LOCTEXT("MediaAssetCategory", "Media")));
 
 	EAssetTypeCategories::Type BlendablesCategoryBit = RegisterAdvancedAssetCategory(FName(TEXT("Blendables")), LOCTEXT("BlendablesAssetCategory", "Blendables"));
+	EAssetTypeCategories::Type FoliageCategoryBit = RegisterAdvancedAssetCategory(FName(TEXT("Foliage")), LOCTEXT("FoliageAssetCategory", "Foliage"));
 
 	// Register the built-in asset type actions
 	RegisterAssetTypeActions(MakeShareable(new FAssetTypeActions_AnimationAsset));
@@ -193,10 +195,11 @@ UAssetToolsImpl::UAssetToolsImpl(const FObjectInitializer& ObjectInitializer)
 	RegisterAssetTypeActions(MakeShareable(new FAssetTypeActions_ForceFeedbackEffect));
 	RegisterAssetTypeActions(MakeShareable(new FAssetTypeActions_HLODProxy));
 	RegisterAssetTypeActions(MakeShareable(new FAssetTypeActions_SubsurfaceProfile));
-	RegisterAssetTypeActions(MakeShareable(new FAssetTypeActions_InstancedFoliageSettings));
+	RegisterAssetTypeActions(MakeShareable(new FAssetTypeActions_InstancedFoliageSettings(FoliageCategoryBit)));
+	RegisterAssetTypeActions(MakeShareable(new FAssetTypeActions_ActorFoliageSettings(FoliageCategoryBit)));
 	RegisterAssetTypeActions(MakeShareable(new FAssetTypeActions_InterpData));
 	RegisterAssetTypeActions(MakeShareable(new FAssetTypeActions_LandscapeLayer));
-	RegisterAssetTypeActions(MakeShareable(new FAssetTypeActions_LandscapeGrassType));
+	RegisterAssetTypeActions(MakeShareable(new FAssetTypeActions_LandscapeGrassType(FoliageCategoryBit)));
 	RegisterAssetTypeActions(MakeShareable(new FAssetTypeActions_Material(BlendablesCategoryBit)));
 	RegisterAssetTypeActions(MakeShareable(new FAssetTypeActions_MaterialFunction));
 	RegisterAssetTypeActions(MakeShareable(new FAssetTypeActions_MaterialFunctionLayer));
@@ -212,7 +215,7 @@ UAssetToolsImpl::UAssetToolsImpl(const FObjectInitializer& ObjectInitializer)
 	RegisterAssetTypeActions(MakeShareable(new FAssetTypeActions_PhysicalMaterial));
 	RegisterAssetTypeActions(MakeShareable(new FAssetTypeActions_PhysicsAsset));
 	RegisterAssetTypeActions(MakeShareable(new FAssetTypeActions_PreviewMeshCollection));
-	RegisterAssetTypeActions(MakeShareable(new FAssetTypeActions_ProceduralFoliageSpawner));
+	RegisterAssetTypeActions(MakeShareable(new FAssetTypeActions_ProceduralFoliageSpawner(FoliageCategoryBit)));
 	RegisterAssetTypeActions(MakeShareable(new FAssetTypeActions_Redirector));
 	RegisterAssetTypeActions(MakeShareable(new FAssetTypeActions_Rig));
 	RegisterAssetTypeActions(MakeShareable(new FAssetTypeActions_SkeletalMesh));
diff --git a/Engine/Source/Developer/AssetTools/Private/AssetTypeActions/AssetTypeActions_ActorFoliageSettings.cpp b/Engine/Source/Developer/AssetTools/Private/AssetTypeActions/AssetTypeActions_ActorFoliageSettings.cpp
new file mode 100644
--- /dev/null
+++ b/Engine/Source/Developer/AssetTools/Private/AssetTypeActions/AssetTypeActions_ActorFoliageSettings.cpp
@@ -0,0 +1,9 @@
+// Copyright 1998-2019 Epic Games, Inc. All Rights Reserved.
+
+#include "AssetTypeActions/AssetTypeActions_ActorFoliageSettings.h"
+#include "FoliageType_Actor.h"
+
+UClass* FAssetTypeActions_ActorFoliageSettings::GetSupportedClass() const
+{
+	return UFoliageType_Actor::StaticClass();
+}
diff --git a/Engine/Source/Developer/AssetTools/Private/AssetTypeActions/AssetTypeActions_ActorFoliageSettings.h b/Engine/Source/Developer/AssetTools/Private/AssetTypeActions/AssetTypeActions_ActorFoliageSettings.h
new file mode 100644
--- /dev/null
+++ b/Engine/Source/Developer/AssetTools/Private/AssetTypeActions/AssetTypeActions_ActorFoliageSettings.h
@@ -0,0 +1,23 @@
+// Copyright 1998-2019 Epic Games, Inc. All Rights Reserved.
+
+#pragma once
+
+#include "CoreMinimal.h"
+#include "AssetTypeActions_Base.h"
+
+class FAssetTypeActions_ActorFoliageSettings : public FAssetTypeActions_Base
+{
+public:
+	FAssetTypeActions_ActorFoliageSettings(EAssetTypeCategories::Type InAssetCategoryBit)
+		: AssetCategoryBit(InAssetCategoryBit)
+	{}
+
+	// IAssetTypeActions Implementation
+	virtual FText GetName() const override { return NSLOCTEXT("AssetTypeActions", "AssetTypeActions_ActorFoliageSettings", "Blueprint Foliage"); }
+	virtual FColor GetTypeColor() const override { return FColor(12, 65, 12); }
+	virtual UClass* GetSupportedClass() const override;
+	virtual uint32 GetCategories() override { return AssetCategoryBit; }
+
+private:
+	EAssetTypeCategories::Type AssetCategoryBit;
+};
diff --git a/Engine/Source/Developer/AssetTools/Private/AssetTypeActions/AssetTypeActions_InstancedFoliageSettings.h b/Engine/Source/Developer/AssetTools/Private/AssetTypeActions/AssetTypeActions_InstancedFoliageSettings.h
--- a/Engine/Source/Developer/AssetTools/Private/AssetTypeActions/AssetTypeActions_InstancedFoliageSettings.h
+++ b/Engine/Source/Developer/AssetTools/Private/AssetTypeActions/AssetTypeActions_InstancedFoliageSettings.h
@@ -8,9 +8,17 @@
 class FAssetTypeActions_InstancedFoliageSettings : public FAssetTypeActions_Base
 {
 public:
+	FAssetTypeActions_InstancedFoliageSettings(EAssetTypeCategories::Type InAssetCategoryBit)
+		: AssetCategoryBit(InAssetCategoryBit)
+	{}
+		
+
 	// IAssetTypeActions Implementation
-	virtual FText GetName() const override { return NSLOCTEXT("AssetTypeActions", "AssetTypeActions_InstancedFoliageSettings", "Foliage Type"); }
+	virtual FText GetName() const override { return NSLOCTEXT("AssetTypeActions", "AssetTypeActions_InstancedFoliageSettings", "Static Mesh Foliage"); }
 	virtual FColor GetTypeColor() const override { return FColor(12, 173, 12); }
 	virtual UClass* GetSupportedClass() const override;
-	virtual uint32 GetCategories() override { return EAssetTypeCategories::Misc; }
+	virtual uint32 GetCategories() override { return AssetCategoryBit; }
+
+private:
+	EAssetTypeCategories::Type AssetCategoryBit;
 };
diff --git a/Engine/Source/Developer/AssetTools/Private/AssetTypeActions/AssetTypeActions_LandscapeGrassType.h b/Engine/Source/Developer/AssetTools/Private/AssetTypeActions/AssetTypeActions_LandscapeGrassType.h
--- a/Engine/Source/Developer/AssetTools/Private/AssetTypeActions/AssetTypeActions_LandscapeGrassType.h
+++ b/Engine/Source/Developer/AssetTools/Private/AssetTypeActions/AssetTypeActions_LandscapeGrassType.h
@@ -9,9 +9,16 @@
 class FAssetTypeActions_LandscapeGrassType : public FAssetTypeActions_Base
 {
 public:
+	FAssetTypeActions_LandscapeGrassType(EAssetTypeCategories::Type InAssetCategoryBit)
+		: AssetCategoryBit(InAssetCategoryBit)
+	{}
+
 	// IAssetTypeActions Implementation
 	virtual FText GetName() const override { return NSLOCTEXT("AssetTypeActions", "AssetTypeActions_LandscapeGrassType", "Landscape Grass Type"); }
 	virtual FColor GetTypeColor() const override { return FColor(128,255,128); }
 	virtual UClass* GetSupportedClass() const override { return ULandscapeGrassType::StaticClass(); }
-	virtual uint32 GetCategories() override { return EAssetTypeCategories::Misc; }
+	virtual uint32 GetCategories() override { return AssetCategoryBit; }
+
+private:
+	EAssetTypeCategories::Type AssetCategoryBit;
 };
diff --git a/Engine/Source/Developer/AssetTools/Private/AssetTypeActions/AssetTypeActions_ProceduralFoliageSpawner.h b/Engine/Source/Developer/AssetTools/Private/AssetTypeActions/AssetTypeActions_ProceduralFoliageSpawner.h
--- a/Engine/Source/Developer/AssetTools/Private/AssetTypeActions/AssetTypeActions_ProceduralFoliageSpawner.h
+++ b/Engine/Source/Developer/AssetTools/Private/AssetTypeActions/AssetTypeActions_ProceduralFoliageSpawner.h
@@ -8,10 +8,17 @@
 class FAssetTypeActions_ProceduralFoliageSpawner : public FAssetTypeActions_Base
 {
 public:
+	FAssetTypeActions_ProceduralFoliageSpawner(EAssetTypeCategories::Type InAssetCategoryBit)
+		: AssetCategoryBit(InAssetCategoryBit)
+	{}
+
 	// IAssetTypeActions Implementation
 	virtual FText GetName() const override { return NSLOCTEXT("AssetTypeActions", "AssetTypeActions_ProceduralFoliageSpawner", "Procedural Foliage Spawner"); }
 	virtual FColor GetTypeColor() const override { return FColor(7, 103, 7); }
 	virtual UClass* GetSupportedClass() const override;
-	virtual uint32 GetCategories() override { return EAssetTypeCategories::Misc; }
+	virtual uint32 GetCategories() override { return AssetCategoryBit; }
 	virtual bool CanFilter() override;
+
+private:
+	EAssetTypeCategories::Type AssetCategoryBit;
 };
diff --git a/Engine/Source/Editor/EditorStyle/Private/SlateEditorStyle.cpp b/Engine/Source/Editor/EditorStyle/Private/SlateEditorStyle.cpp
--- a/Engine/Source/Editor/EditorStyle/Private/SlateEditorStyle.cpp
+++ b/Engine/Source/Editor/EditorStyle/Private/SlateEditorStyle.cpp
@@ -6339,7 +6339,8 @@ void FSlateEditorStyle::FStyle::SetupClassIconsAndThumbnails()
 			TEXT("Plane"),
 			TEXT("CineCameraActor"),
 			TEXT("CameraRig_Crane"),
-			TEXT("CameraRig_Rail")
+			TEXT("CameraRig_Rail"),
+			TEXT("FoliageType_Actor")
 		};
 
 		for (int32 TypeIndex = 0; TypeIndex < ARRAY_COUNT(AssetTypes); ++TypeIndex)
diff --git a/Engine/Source/Editor/FoliageEdit/Private/FoliageEdMode.cpp b/Engine/Source/Editor/FoliageEdit/Private/FoliageEdMode.cpp
--- a/Engine/Source/Editor/FoliageEdit/Private/FoliageEdMode.cpp
+++ b/Engine/Source/Editor/FoliageEdit/Private/FoliageEdMode.cpp
@@ -23,6 +23,7 @@
 #include "EditorModeManager.h"
 #include "FileHelpers.h"
 #include "ScopedTransaction.h"
+#include "Engine/Selection.h"
 
 #include "FoliageEdModeToolkit.h"
 #include "LevelEditor.h"
@@ -87,22 +88,22 @@ FText FFoliageMeshUIInfo::GetNameText() const
 }
 
 //
-// FFoliageMeshInfo iterator
+// FFoliageInfo iterator
 //
-class FFoliageMeshInfoIterator
+class FFoliageInfoIterator
 {
 private:
 	const UWorld*				World;
 	const UFoliageType*			FoliageType;
-	FFoliageMeshInfo*			CurrentMeshInfo;
+	FFoliageInfo*			CurrentInfo;
 	AInstancedFoliageActor*		CurrentIFA;
 	int32						LevelIdx;
 
 public:
-	FFoliageMeshInfoIterator(const UWorld* InWorld, const UFoliageType* InFoliageType)
+	FFoliageInfoIterator(const UWorld* InWorld, const UFoliageType* InFoliageType)
 		: World(InWorld)
 		, FoliageType(InFoliageType)
-		, CurrentMeshInfo(nullptr)
+		, CurrentInfo(nullptr)
 		, CurrentIFA(nullptr)
 	{
 		// shortcut for non-assets
@@ -113,7 +114,7 @@ public:
 			if (IFA->GetLevel()->bIsVisible)
 			{
 				CurrentIFA = IFA;
-				CurrentMeshInfo = CurrentIFA->FindMesh(FoliageType);
+				CurrentInfo = CurrentIFA->FindInfo(FoliageType);
 			}
 		}
 		else
@@ -136,8 +137,8 @@ public:
 				CurrentIFA = AInstancedFoliageActor::GetInstancedFoliageActorForLevel(Level);
 				if (CurrentIFA)
 				{
-					CurrentMeshInfo = CurrentIFA->FindMesh(FoliageType);
-					if (CurrentMeshInfo)
+					CurrentInfo = CurrentIFA->FindInfo(FoliageType);
+					if (CurrentInfo)
 					{
 						LevelIdx = LocalLevelIdx;
 						return;
@@ -146,19 +147,19 @@ public:
 			}
 		}
 
-		CurrentMeshInfo = nullptr;
+		CurrentInfo = nullptr;
 		CurrentIFA = nullptr;
 	}
 
-	FORCEINLINE FFoliageMeshInfo* operator*()
+	FORCEINLINE FFoliageInfo* operator*()
 	{
-		check(CurrentMeshInfo);
-		return CurrentMeshInfo;
+		check(CurrentInfo);
+		return CurrentInfo;
 	}
 
 	FORCEINLINE operator bool() const
 	{
-		return CurrentMeshInfo != nullptr;
+		return CurrentInfo != nullptr;
 	}
 
 	FORCEINLINE AInstancedFoliageActor* GetActor()
@@ -174,12 +175,14 @@ bool FFoliagePaintingGeometryFilter::operator() (const UPrimitiveComponent* Comp
 {
 	if (Component)
 	{
+		bool bFoliageOwned = Component->GetOwner() && AInstancedFoliageActor::IsOwnedByFoliage(Component->GetOwner());
+
 		// Whitelist
 		bool bAllowed =
 			(bAllowLandscape   && Component->IsA(ULandscapeHeightfieldCollisionComponent::StaticClass())) ||
-			(bAllowStaticMesh  && Component->IsA(UStaticMeshComponent::StaticClass()) && !Component->IsA(UFoliageInstancedStaticMeshComponent::StaticClass())) ||
+			(bAllowStaticMesh  && Component->IsA(UStaticMeshComponent::StaticClass()) && !Component->IsA(UFoliageInstancedStaticMeshComponent::StaticClass()) && !bFoliageOwned) ||
 			(bAllowBSP         && (Component->IsA(UBrushComponent::StaticClass()) || Component->IsA(UModelComponent::StaticClass()))) ||
-			(bAllowFoliage     && Component->IsA(UFoliageInstancedStaticMeshComponent::StaticClass()));
+			(bAllowFoliage     && (Component->IsA(UFoliageInstancedStaticMeshComponent::StaticClass()) || bFoliageOwned));
 
 		// Blacklist
 		bAllowed &=
@@ -343,6 +346,7 @@ void FEdModeFoliage::Enter()
 	FEditorDelegates::NewCurrentLevel.AddSP(this, &FEdModeFoliage::NotifyNewCurrentLevel);
 	FWorldDelegates::LevelAddedToWorld.AddSP(this, &FEdModeFoliage::NotifyLevelAddedToWorld);
 	FWorldDelegates::LevelRemovedFromWorld.AddSP(this, &FEdModeFoliage::NotifyLevelRemovedFromWorld);
+	AInstancedFoliageActor::SelectionChanged.AddSP(this, &FEdModeFoliage::NotifyActorSelectionChanged);
 
 	FAssetRegistryModule& AssetRegistryModule = FModuleManager::GetModuleChecked<FAssetRegistryModule>(TEXT("AssetRegistry"));
 	AssetRegistryModule.Get().OnAssetRemoved().AddSP(this, &FEdModeFoliage::NotifyAssetRemoved);
@@ -361,7 +365,7 @@ void FEdModeFoliage::Enter()
 
 	if (UISettings.GetSelectToolSelected() || UISettings.GetLassoSelectToolSelected())
 	{
-		ApplySelectionToComponents(GetWorld(), true);
+		ApplySelection(GetWorld(), true);
 	}
 
 	TArray<AInstancedFoliageActor*> InstanceFoliageActorList;
@@ -388,20 +392,13 @@ void FEdModeFoliage::Enter()
 
 	// Update UI
 	NotifyNewCurrentLevel();
-
-	// Make sure we're up to date and register static mesh bounds changes
-	for (FFoliageMeshUIInfoPtr MeshUIInfo : FoliageMeshList)
+		
+	// Disable foliage engine scalability in foliage mode
+	for (AInstancedFoliageActor* Actor : InstanceFoliageActorList)
 	{
-		for (AInstancedFoliageActor* Actor : InstanceFoliageActorList)
+		for (auto& FoliageMesh : Actor->FoliageInfos)
 		{
-			FFoliageMeshInfo* MeshInfo = Actor->FindMesh(MeshUIInfo->Settings);
-
-			if (MeshInfo != nullptr && MeshInfo->Component != nullptr && MeshInfo->Component->GetStaticMesh() != nullptr)
-			{
-				MeshInfo->Component->GetStaticMesh()->GetOnExtendedBoundsChanged().AddRaw(MeshInfo, &FFoliageMeshInfo::HandleComponentMeshBoundsChanged);
-
-				MeshInfo->Component->BuildTreeIfOutdated(true, false);
-			}
+			FoliageMesh.Value->EnterEditMode();
 		}
 	}
 
@@ -423,19 +420,6 @@ void FEdModeFoliage::Enter()
 
 	// Make sure the brush is visible.
 	SphereBrushComponent->SetVisibility(true);
-
-	// Disable foliage engine scalability in foliage mode
-	for (AInstancedFoliageActor* Actor : InstanceFoliageActorList)
-	{
-		for (auto& FoliageMesh : Actor->FoliageMeshes)
-		{
-			if (FoliageMesh.Value->Component != nullptr)
-			{
-				FoliageMesh.Value->Component->bCanEnableDensityScaling = false;
-				FoliageMesh.Value->Component->UpdateDensityScaling();
-			}
-		}		
-	}
 }
 
 /** FEdMode: Called when the mode is exited */
@@ -465,6 +449,7 @@ void FEdModeFoliage::Exit()
 	FEditorDelegates::NewCurrentLevel.RemoveAll(this);
 	FWorldDelegates::LevelAddedToWorld.RemoveAll(this);
 	FWorldDelegates::LevelRemovedFromWorld.RemoveAll(this);
+	AInstancedFoliageActor::SelectionChanged.RemoveAll(this);
 
 	if (FModuleManager::Get().IsModuleLoaded(TEXT("AssetRegistry")))
 	{
@@ -490,7 +475,7 @@ void FEdModeFoliage::Exit()
 	UISettings.Save();
 
 	// Clear selection visualization on any foliage items
-	ApplySelectionToComponents(GetWorld(), false);
+	ApplySelection(GetWorld(), false);
 
 	// Remove all event subscriptions
 	TArray<AInstancedFoliageActor*> InstanceFoliageActorList;
@@ -512,48 +497,17 @@ void FEdModeFoliage::Exit()
 			}
 		}
 	}
-
-	// Make sure we're up to date and register static mesh bounds changes
-	for (FFoliageMeshUIInfoPtr MeshUIInfo : FoliageMeshList)
-	{
-		for (AInstancedFoliageActor* Actor : InstanceFoliageActorList)
-		{
-			FFoliageMeshInfo* MeshInfo = Actor->FindMesh(MeshUIInfo->Settings);
-
-			if (MeshInfo != nullptr && MeshInfo->Component != nullptr && MeshInfo->Component->GetStaticMesh() != nullptr)
-			{
-				MeshInfo->Component->GetStaticMesh()->GetOnExtendedBoundsChanged().RemoveAll(MeshInfo);
-			}
-		}
-	}
-
+		
 	for (AInstancedFoliageActor* Actor : InstanceFoliageActorList)
 	{
-		for (auto& FoliageMesh : Actor->FoliageMeshes)
+		for (auto& FoliageMesh : Actor->FoliageInfos)
 		{
-			if (FoliageMesh.Value->Component != nullptr)
-			{
-				FoliageMesh.Value->Component->bCanEnableDensityScaling = true;
-				FoliageMesh.Value->Component->UpdateDensityScaling();
-			}
+			FoliageMesh.Value->ExitEditMode();	
 		}
 	}
 
 	FEditorDelegates::EndPIE.RemoveAll(this);
-
-	// Reenable foliage engine scalability
-	for (AInstancedFoliageActor* Actor : InstanceFoliageActorList)
-	{
-		for (auto& FoliageMesh : Actor->FoliageMeshes)
-		{
-			if (FoliageMesh.Value->Component != nullptr)
-			{
-				FoliageMesh.Value->Component->bCanEnableDensityScaling = true;
-				FoliageMesh.Value->Component->UpdateDensityScaling();
-			}
-		}
-	}
-
+		
 	FoliageMeshList.Empty();
 
 	// Call base Exit method to ensure proper cleanup
@@ -760,16 +714,25 @@ void FEdModeFoliage::NotifyAssetRemoved(const FAssetData& AssetInfo)
 	}
 }
 
+void FEdModeFoliage::NotifyActorSelectionChanged(bool bSelect, const TArray<AActor*>& Selection)
+{
+	GEditor->GetSelectedActors()->Modify();
+	for (AActor* Actor : Selection)
+	{
+		GEditor->SelectActor(Actor, bSelect, true, true);
+	}
+}
+
 /** When the user changes the current tool in the UI */
 void FEdModeFoliage::HandleToolChanged()
 {
 	if (UISettings.GetSelectToolSelected() || UISettings.GetLassoSelectToolSelected())
 	{
-		ApplySelectionToComponents(GetWorld(), true);
+		ApplySelection(GetWorld(), true);
 	}
 	else
 	{
-		ApplySelectionToComponents(GetWorld(), false);
+		ApplySelection(GetWorld(), false);
 	}
 
 	OnToolChanged.Broadcast();
@@ -846,11 +809,11 @@ void FEdModeFoliage::OnObjectsReplaced(const TMap<UObject*, UObject*>& Replaceme
 				{
 					if (UFoliageType* ReplacedFoliageType = Cast<UFoliageType>(ReplacementPair.Key))
 					{
-						TUniqueObj<FFoliageMeshInfo> MeshInfo;
-						if (IFA->FoliageMeshes.RemoveAndCopyValue(ReplacedFoliageType, MeshInfo))
+						TUniqueObj<FFoliageInfo> FoliageInfo;
+						if (IFA->FoliageInfos.RemoveAndCopyValue(ReplacedFoliageType, FoliageInfo))
 						{
 							// Re-add the unique mesh info associated with the replaced foliage type
-							IFA->FoliageMeshes.Add(Cast<UFoliageType>(ReplacementPair.Value), MoveTemp(MeshInfo));
+							IFA->FoliageInfos.Add(Cast<UFoliageType>(ReplacementPair.Value), MoveTemp(FoliageInfo));
 							bAnyFoliageTypeReplaced = true;
 						}
 					}
@@ -1109,10 +1072,10 @@ static bool CheckForOverlappingSphere(AInstancedFoliageActor* IFA, const UFoliag
 {
 	if (IFA)
 	{
-		FFoliageMeshInfo* MeshInfo = IFA->FindMesh(Settings);
-		if (MeshInfo)
+		FFoliageInfo* Info = IFA->FindInfo(Settings);
+		if (Info)
 		{
-			return MeshInfo->CheckForOverlappingSphere(Sphere);
+			return Info->CheckForOverlappingSphere(Sphere);
 		}
 	}
 
@@ -1122,10 +1085,10 @@ static bool CheckForOverlappingSphere(AInstancedFoliageActor* IFA, const UFoliag
 // Returns whether or not there is are any instances overlapping the sphere specified
 static bool CheckForOverlappingSphere(const UWorld* InWorld, const UFoliageType* Settings, const FSphere& Sphere)
 {
-	for (FFoliageMeshInfoIterator It(InWorld, Settings); It; ++It)
+	for (FFoliageInfoIterator It(InWorld, Settings); It; ++It)
 	{
-		FFoliageMeshInfo* MeshInfo = (*It);
-		if (MeshInfo->CheckForOverlappingSphere(Sphere))
+		FFoliageInfo* Info = (*It);
+		if (Info->CheckForOverlappingSphere(Sphere))
 		{
 			return true;
 		}
@@ -1472,14 +1435,13 @@ static void SpawnFoliageInstance(UWorld* InWorld, const UFoliageType* Settings,
 		CurrentFoliageTraceBrushAffectedLevels.AddUnique(TargetLevel);
 		AInstancedFoliageActor* IFA = AInstancedFoliageActor::GetInstancedFoliageActorForLevel(TargetLevel, true);
 
-		FFoliageMeshInfo* MeshInfo = nullptr;
-		UFoliageType* FoliageSettings = IFA->AddFoliageType(Settings, &MeshInfo);
-
-		MeshInfo->AddInstances(IFA, FoliageSettings, PlacedLevelInstances.Value, false);
+		FFoliageInfo* Info = nullptr;
+		UFoliageType* FoliageSettings = IFA->AddFoliageType(Settings, &Info);
 
-		if (MeshInfo->Component != nullptr && InRebuildFoliageTree)
+		Info->AddInstances(IFA, FoliageSettings, PlacedLevelInstances.Value);
+		if (InRebuildFoliageTree)
 		{
-			MeshInfo->Component->BuildTreeIfOutdated(true, false);
+			Info->Refresh(IFA, true, false);
 		}
 	}
 }
@@ -1492,11 +1454,11 @@ void FEdModeFoliage::RebuildFoliageTree(const UFoliageType* Settings)
 
 		if (IFA != nullptr)
 		{
-			FFoliageMeshInfo* MeshInfo = IFA->FindMesh(Settings);
+			FFoliageInfo* FoliageInfo = IFA->FindInfo(Settings);
 
-			if (MeshInfo != nullptr && MeshInfo->Component != nullptr)
+			if (FoliageInfo != nullptr)
 			{
-				MeshInfo->Component->BuildTreeIfOutdated(true, false);
+				FoliageInfo->Refresh(IFA, true, false);
 			}
 		}
 	}
@@ -1606,11 +1568,11 @@ void FEdModeFoliage::AddInstancesForBrush(UWorld* InWorld, const UFoliageType* S
 	ExistingInstanceBuckets.AddZeroed(NUM_INSTANCE_BUCKETS);
 	int32 NumExistingInstances = 0;
 
-	for (FFoliageMeshInfoIterator It(World, Settings); It; ++It)
+	for (FFoliageInfoIterator It(World, Settings); It; ++It)
 	{
-		FFoliageMeshInfo* MeshInfo = (*It);
+		FFoliageInfo* FoliageInfo = (*It);
 		TArray<int32> ExistingInstances;
-		MeshInfo->GetInstancesInsideSphere(BrushSphere, ExistingInstances);
+		FoliageInfo->GetInstancesInsideSphere(BrushSphere, ExistingInstances);
 		NumExistingInstances += ExistingInstances.Num();
 
 		if (bHasValidLandscapeLayers)
@@ -1618,7 +1580,7 @@ void FEdModeFoliage::AddInstancesForBrush(UWorld* InWorld, const UFoliageType* S
 			// Find the landscape weights of existing ExistingInstances
 			for (int32 Idx : ExistingInstances)
 			{
-				FFoliageInstance& Instance = MeshInfo->Instances[Idx];
+				FFoliageInstance& Instance = FoliageInfo->Instances[Idx];
 				auto InstanceBasePtr = It.GetActor()->InstanceBaseCache.GetInstanceBasePtr(Instance.BaseId);
 				float HitWeight;
 				if (GetMaxHitWeight(Instance.Location, InstanceBasePtr.Get(), Settings, &LandscapeLayerCaches, HitWeight))
@@ -1656,13 +1618,13 @@ void FEdModeFoliage::RemoveInstancesForBrush(UWorld* InWorld, const UFoliageType
 {
 	SCOPE_CYCLE_COUNTER(STAT_FoliageRemoveInstanceBrush);
 
-	for (FFoliageMeshInfoIterator It(InWorld, Settings); It; ++It)
+	for (FFoliageInfoIterator It(InWorld, Settings); It; ++It)
 	{
-		FFoliageMeshInfo* MeshInfo = (*It);
+		FFoliageInfo* FoliageInfo = (*It);
 		AInstancedFoliageActor* IFA = It.GetActor();
 
 		TArray<int32> PotentialInstancesToRemove;
-		MeshInfo->GetInstancesInsideSphere(BrushSphere, PotentialInstancesToRemove);
+		FoliageInfo->GetInstancesInsideSphere(BrushSphere, PotentialInstancesToRemove);
 		if (PotentialInstancesToRemove.Num() == 0)
 		{
 			continue;
@@ -1689,7 +1651,7 @@ void FEdModeFoliage::RemoveInstancesForBrush(UWorld* InWorld, const UFoliageType
 		// Filter PotentialInstancesToRemove
 		for (int32 Idx = 0; Idx < PotentialInstancesToRemove.Num(); Idx++)
 		{
-			auto BaseId = MeshInfo->Instances[PotentialInstancesToRemove[Idx]].BaseId;
+			auto BaseId = FoliageInfo->Instances[PotentialInstancesToRemove[Idx]].BaseId;
 			auto BasePtr = IFA->InstanceBaseCache.GetInstanceBasePtr(BaseId);
 			UPrimitiveComponent* Base = Cast<UPrimitiveComponent>(BasePtr.Get());
 
@@ -1707,7 +1669,7 @@ void FEdModeFoliage::RemoveInstancesForBrush(UWorld* InWorld, const UFoliageType
 		{
 			CurrentFoliageTraceBrushAffectedLevels.AddUnique(IFA->GetLevel());
 
-			MeshInfo->RemoveInstances(IFA, PotentialInstancesToRemove, false);
+			FoliageInfo->RemoveInstances(IFA, PotentialInstancesToRemove, false);
 		}
 	}
 }
@@ -1715,38 +1677,38 @@ void FEdModeFoliage::RemoveInstancesForBrush(UWorld* InWorld, const UFoliageType
 
 void FEdModeFoliage::SelectInstanceAtLocation(UWorld* InWorld, const UFoliageType* Settings, const FVector& Location, bool bSelect)
 {
-	for (FFoliageMeshInfoIterator It(InWorld, Settings); It; ++It)
+	for (FFoliageInfoIterator It(InWorld, Settings); It; ++It)
 	{
-		FFoliageMeshInfo* MeshInfo = (*It);
+		FFoliageInfo* FoliageInfo = (*It);
 		AInstancedFoliageActor* IFA = It.GetActor();
 
 		int32 Instance;
 		bool bResult;
-		MeshInfo->GetInstanceAtLocation(Location, Instance, bResult);
+		FoliageInfo->GetInstanceAtLocation(Location, Instance, bResult);
 		if (bResult)
 		{
-		        TArray<int32> Instances;
+		    TArray<int32> Instances;
 			Instances.Add(Instance);
-			MeshInfo->SelectInstances(IFA, bSelect, Instances);
+			FoliageInfo->SelectInstances(IFA, bSelect, Instances);
 		}
 	}
 }
 
 void FEdModeFoliage::SelectInstancesForBrush(UWorld* InWorld, const UFoliageType* Settings, const FSphere& BrushSphere, bool bSelect)
 {
-	for (FFoliageMeshInfoIterator It(InWorld, Settings); It; ++It)
+	for (FFoliageInfoIterator It(InWorld, Settings); It; ++It)
 	{
-		FFoliageMeshInfo* MeshInfo = (*It);
+		FFoliageInfo* FoliageInfo = (*It);
 		AInstancedFoliageActor* IFA = It.GetActor();
 
 		TArray<int32> Instances;
-		MeshInfo->GetInstancesInsideSphere(BrushSphere, Instances);
+		FoliageInfo->GetInstancesInsideSphere(BrushSphere, Instances);
 		if (Instances.Num() == 0)
 		{
 			continue;
 		}
 
-		MeshInfo->SelectInstances(IFA, bSelect, Instances);
+		FoliageInfo->SelectInstances(IFA, bSelect, Instances);
 	}
 }
 
@@ -1772,17 +1734,19 @@ void FEdModeFoliage::SelectInstances(UWorld* InWorld, bool bSelect)
 
 void FEdModeFoliage::SelectInstances(UWorld* InWorld, const UFoliageType* Settings, bool bSelect)
 {
-	for (FFoliageMeshInfoIterator It(InWorld, Settings); It; ++It)
+	for (FFoliageInfoIterator It(InWorld, Settings); It; ++It)
 	{
-		FFoliageMeshInfo* MeshInfo = (*It);
+		FFoliageInfo* FoliageInfo = (*It);
 		AInstancedFoliageActor* IFA = It.GetActor();
 
-		MeshInfo->SelectInstances(IFA, bSelect);
+		FoliageInfo->SelectInstances(IFA, bSelect);
 	}
 }
 
-void FEdModeFoliage::ApplySelectionToComponents(UWorld* InWorld, bool bApply)
+void FEdModeFoliage::ApplySelection(UWorld* InWorld, bool bApply)
 {
+	GEditor->SelectNone(true, true);
+	
 	const int32 NumLevels = InWorld->GetNumLevels();
 	for (int32 LevelIdx = 0; LevelIdx < NumLevels; ++LevelIdx)
 	{
@@ -1790,7 +1754,7 @@ void FEdModeFoliage::ApplySelectionToComponents(UWorld* InWorld, bool bApply)
 		AInstancedFoliageActor* IFA = AInstancedFoliageActor::GetInstancedFoliageActorForLevel(Level);
 		if (IFA)
 		{
-			IFA->ApplySelectionToComponents(bApply);
+			IFA->ApplySelection(bApply);
 		}
 	}
 }
@@ -1806,10 +1770,10 @@ void FEdModeFoliage::TransformSelectedInstances(UWorld* InWorld, const FVector&
 		{
 			bool bFoundSelection = false;
 
-			for (auto& MeshPair : IFA->FoliageMeshes)
+			for (auto& MeshPair : IFA->FoliageInfos)
 			{
-				FFoliageMeshInfo& MeshInfo = *MeshPair.Value;
-				TArray<int32> SelectedIndices = MeshInfo.SelectedIndices.Array();
+				FFoliageInfo& FoliageInfo = *MeshPair.Value;
+				TArray<int32> SelectedIndices = FoliageInfo.SelectedIndices.Array();
 
 				if (SelectedIndices.Num() > 0)
 				{
@@ -1822,22 +1786,22 @@ void FEdModeFoliage::TransformSelectedInstances(UWorld* InWorld, const FVector&
 
 					if (bDuplicate)
 					{
-						MeshInfo.DuplicateInstances(IFA, MeshPair.Key, SelectedIndices);
+						FoliageInfo.DuplicateInstances(IFA, MeshPair.Key, SelectedIndices);
 						OnInstanceCountUpdated(MeshPair.Key);
 					}
 
-					MeshInfo.PreMoveInstances(IFA, SelectedIndices);
+					FoliageInfo.PreMoveInstances(IFA, SelectedIndices);
 
 					for (int32 SelectedInstanceIdx : SelectedIndices)
 					{
-						FFoliageInstance& Instance = MeshInfo.Instances[SelectedInstanceIdx];
+						FFoliageInstance& Instance = FoliageInfo.Instances[SelectedInstanceIdx];
 						Instance.Location += InDrag;
 						Instance.ZOffset = 0.f;
 						Instance.Rotation += InRot;
 						Instance.DrawScale3D += InScale;
 					}
 
-					MeshInfo.PostMoveInstances(IFA, SelectedIndices);
+					FoliageInfo.PostMoveInstances(IFA, SelectedIndices);
 				}
 			}
 
@@ -1901,7 +1865,7 @@ void FEdModeFoliage::RemoveSelectedInstances(UWorld* InWorld)
 		if (IFA)
 		{
 			bool bHasSelection = false;
-			for (auto& MeshPair : IFA->FoliageMeshes)
+			for (auto& MeshPair : IFA->FoliageInfos)
 			{
 				if (MeshPair.Value->SelectedIndices.Num())
 				{
@@ -1913,15 +1877,14 @@ void FEdModeFoliage::RemoveSelectedInstances(UWorld* InWorld)
 			if (bHasSelection)
 			{
 				IFA->Modify();
-				for (auto& MeshPair : IFA->FoliageMeshes)
+				for (auto& MeshPair : IFA->FoliageInfos)
 				{
-					FFoliageMeshInfo& Mesh = *MeshPair.Value;
+					FFoliageInfo& Mesh = *MeshPair.Value;
 					if (Mesh.SelectedIndices.Num() > 0)
 					{
 						TArray<int32> InstancesToDelete = Mesh.SelectedIndices.Array();
-						Mesh.RemoveInstances(IFA, InstancesToDelete, false);
-						Mesh.Component->BuildTreeIfOutdated(true, true);
-
+						Mesh.RemoveInstances(IFA, InstancesToDelete, true);
+					
 						OnInstanceCountUpdated(MeshPair.Key);
 					}
 				}
@@ -1947,18 +1910,18 @@ void FEdModeFoliage::SelectInvalidInstances(const UFoliageType* Settings)
 	SphereShape.SetSphere(0.f);
 	float InstanceOffGroundLocalThreshold = CVarOffGroundTreshold.GetValueOnGameThread();
 
-	for (FFoliageMeshInfoIterator It(InWorld, Settings); It; ++It)
+	for (FFoliageInfoIterator It(InWorld, Settings); It; ++It)
 	{
-		FFoliageMeshInfo* MeshInfo = (*It);
+		FFoliageInfo* FoliageInfo = (*It);
 		AInstancedFoliageActor* IFA = It.GetActor();
-		int32 NumInstances = MeshInfo->Instances.Num();
+		int32 NumInstances = FoliageInfo->Instances.Num();
 		TArray<FHitResult> Hits; Hits.Reserve(16);
 
 		TArray<int32> InvalidInstances;
 
 		for (int32 InstanceIdx = 0; InstanceIdx < NumInstances; ++InstanceIdx)
 		{
-			FFoliageInstance& Instance = MeshInfo->Instances[InstanceIdx];
+			FFoliageInstance& Instance = FoliageInfo->Instances[InstanceIdx];
 			UActorComponent* CurrentInstanceBase = IFA->InstanceBaseCache.GetInstanceBasePtr(Instance.BaseId).Get();
 			bool bInvalidInstance = true;
 
@@ -2012,7 +1975,7 @@ void FEdModeFoliage::SelectInvalidInstances(const UFoliageType* Settings)
 
 		if (InvalidInstances.Num() > 0)
 		{
-			MeshInfo->SelectInstances(IFA, true, InvalidInstances);
+			FoliageInfo->SelectInstances(IFA, true, InvalidInstances);
 		}
 	}
 }
@@ -2041,20 +2004,20 @@ void FEdModeFoliage::ReapplyInstancesForBrush(UWorld* InWorld, const UFoliageTyp
 	// Adjust instance density first
 	ReapplyInstancesDensityForBrush(InWorld, Settings, BrushSphere, Pressure);
 
-	for (FFoliageMeshInfoIterator It(InWorld, Settings); It; ++It)
+	for (FFoliageInfoIterator It(InWorld, Settings); It; ++It)
 	{
-		FFoliageMeshInfo* MeshInfo = (*It);
+		FFoliageInfo* FoliageInfo = (*It);
 		AInstancedFoliageActor* IFA = It.GetActor();
 
-		ReapplyInstancesForBrush(InWorld, IFA, Settings, MeshInfo, BrushSphere, Pressure);
+		ReapplyInstancesForBrush(InWorld, IFA, Settings, FoliageInfo, BrushSphere, Pressure);
 	}
 }
 
 /** Reapply instance settings to exiting instances */
-void FEdModeFoliage::ReapplyInstancesForBrush(UWorld* InWorld, AInstancedFoliageActor* IFA, const UFoliageType* Settings, FFoliageMeshInfo* MeshInfo, const FSphere& BrushSphere, float Pressure)
+void FEdModeFoliage::ReapplyInstancesForBrush(UWorld* InWorld, AInstancedFoliageActor* IFA, const UFoliageType* Settings, FFoliageInfo* FoliageInfo, const FSphere& BrushSphere, float Pressure)
 {
 	TArray<int32> ExistingInstances;
-	MeshInfo->GetInstancesInsideSphere(BrushSphere, ExistingInstances);
+	FoliageInfo->GetInstancesInsideSphere(BrushSphere, ExistingInstances);
 
 	bool bUpdated = false;
 	TArray<int32> UpdatedInstances;
@@ -2065,7 +2028,7 @@ void FEdModeFoliage::ReapplyInstancesForBrush(UWorld* InWorld, AInstancedFoliage
 	for (int32 Idx = 0; Idx < ExistingInstances.Num(); Idx++)
 	{
 		int32 InstanceIndex = ExistingInstances[Idx];
-		FFoliageInstance& Instance = MeshInfo->Instances[InstanceIndex];
+		FFoliageInstance& Instance = FoliageInfo->Instances[InstanceIndex];
 
 		if ((Instance.Flags & FOLIAGE_Readjusted) == 0)
 		{
@@ -2304,14 +2267,14 @@ void FEdModeFoliage::ReapplyInstancesForBrush(UWorld* InWorld, AInstancedFoliage
 			// Update the hash
 			if (bReapplyLocation)
 			{
-				MeshInfo->InstanceHash->RemoveInstance(OldInstanceLocation, InstanceIndex);
-				MeshInfo->InstanceHash->InsertInstance(Instance.Location, InstanceIndex);
+				FoliageInfo->InstanceHash->RemoveInstance(OldInstanceLocation, InstanceIndex);
+				FoliageInfo->InstanceHash->InsertInstance(Instance.Location, InstanceIndex);
 			}
 
 			// Cull overlapping based on radius
 			if (Settings->ReapplyRadius && Settings->Radius > 0.f)
 			{
-				if (MeshInfo->CheckForOverlappingInstanceExcluding(InstanceIndex, Settings->Radius, InstancesToDelete))
+				if (FoliageInfo->CheckForOverlappingInstanceExcluding(InstanceIndex, Settings->Radius, InstancesToDelete))
 				{
 					InstancesToDelete.Add(InstanceIndex);
 					continue;
@@ -2348,13 +2311,13 @@ void FEdModeFoliage::ReapplyInstancesForBrush(UWorld* InWorld, AInstancedFoliage
 
 	if (UpdatedInstances.Num() > 0)
 	{
-		MeshInfo->PostUpdateInstances(IFA, UpdatedInstances);
+		FoliageInfo->PostUpdateInstances(IFA, UpdatedInstances);
 		IFA->RegisterAllComponents();
 	}
 
 	if (InstancesToDelete.Num())
 	{
-		MeshInfo->RemoveInstances(IFA, InstancesToDelete.Array(), true);
+		FoliageInfo->RemoveInstances(IFA, InstancesToDelete.Array(), true);
 	}
 }
 
@@ -2403,15 +2366,15 @@ void FEdModeFoliage::PreApplyBrush()
 				continue;
 			}
 
-			for (FFoliageMeshInfoIterator It(World, Settings); It; ++It)
+			for (FFoliageInfoIterator It(World, Settings); It; ++It)
 			{
-				FFoliageMeshInfo* MeshInfo = (*It);
+				FFoliageInfo* FoliageInfo = (*It);
 
 				// Take a snapshot of all the locations
-				InstanceSnapshot.Add(Settings, FMeshInfoSnapshot(MeshInfo));
+				InstanceSnapshot.Add(Settings, FMeshInfoSnapshot(FoliageInfo));
 
 				// Clear the "FOLIAGE_Readjusted" flag
-				for (auto& Instance : MeshInfo->Instances)
+				for (auto& Instance : FoliageInfo->Instances)
 				{
 					Instance.Flags &= (~FOLIAGE_Readjusted);
 				}
@@ -2554,7 +2517,7 @@ void FEdModeFoliage::ApplyPaintBucket_Remove(AActor* Actor)
 		}
 
 		// Go through all FoliageActors in the world and delete 
-		for (FFoliageMeshInfoIterator It(World, FoliageType); It; ++It)
+		for (FFoliageInfoIterator It(World, FoliageType); It; ++It)
 		{
 			AInstancedFoliageActor* IFA = It.GetActor();
 
@@ -2835,10 +2798,10 @@ void FEdModeFoliage::SnapSelectedInstancesToGround(UWorld* InWorld)
 			{
 				bool bFoundSelection = false;
 
-				for (auto& MeshPair : IFA->FoliageMeshes)
+				for (auto& MeshPair : IFA->FoliageInfos)
 				{
-					FFoliageMeshInfo& MeshInfo = *MeshPair.Value;
-					TArray<int32> SelectedIndices = MeshInfo.SelectedIndices.Array();
+					FFoliageInfo& FoliageInfo = *MeshPair.Value;
+					TArray<int32> SelectedIndices = FoliageInfo.SelectedIndices.Array();
 
 					if (SelectedIndices.Num() > 0)
 					{
@@ -2849,14 +2812,14 @@ void FEdModeFoliage::SnapSelectedInstancesToGround(UWorld* InWorld)
 							bFoundSelection = true;
 						}
 
-						MeshInfo.PreMoveInstances(IFA, SelectedIndices);
+						FoliageInfo.PreMoveInstances(IFA, SelectedIndices);
 
 						for (int32 InstanceIndex : SelectedIndices)
 						{
-							bMovedInstance |= SnapInstanceToGround(IFA, MeshPair.Key->AlignMaxAngle, MeshInfo, InstanceIndex);
+							bMovedInstance |= SnapInstanceToGround(IFA, MeshPair.Key->AlignMaxAngle, FoliageInfo, InstanceIndex);
 						}
 
-						MeshInfo.PostMoveInstances(IFA, SelectedIndices);
+						FoliageInfo.PostMoveInstances(IFA, SelectedIndices);
 					}
 				}
 			}
@@ -2881,7 +2844,7 @@ void FEdModeFoliage::HandleOnActorSpawned(AActor* Actor)
 
 void FEdModeFoliage::HandleOnFoliageTypeMeshChanged(UFoliageType* FoliageType)
 {
-	if (FoliageType->IsNotAssetOrBlueprint() && FoliageType->GetStaticMesh() == nullptr)
+	if (FoliageType->IsNotAssetOrBlueprint() && FoliageType->GetSource() == nullptr)
 	{
 		RemoveFoliageType(&FoliageType, 1);
 	}
@@ -2891,7 +2854,7 @@ void FEdModeFoliage::HandleOnFoliageTypeMeshChanged(UFoliageType* FoliageType)
 	}
 }
 
-bool FEdModeFoliage::SnapInstanceToGround(AInstancedFoliageActor* InIFA, float AlignMaxAngle, FFoliageMeshInfo& Mesh, int32 InstanceIdx)
+bool FEdModeFoliage::SnapInstanceToGround(AInstancedFoliageActor* InIFA, float AlignMaxAngle, FFoliageInfo& Mesh, int32 InstanceIdx)
 {
 	FFoliageInstance& Instance = Mesh.Instances[InstanceIdx];
 	FVector Start = Instance.Location;
@@ -2969,7 +2932,7 @@ void FEdModeFoliage::PopulateFoliageMeshList()
 			AInstancedFoliageActor* IFA = AInstancedFoliageActor::GetInstancedFoliageActorForLevel(Level);
 			if (IFA)
 			{
-				for (auto& MeshPair : IFA->FoliageMeshes)
+				for (auto& MeshPair : IFA->FoliageInfos)
 				{
 					if (!CanPaint(MeshPair.Key, CurrentLevel))
 					{
@@ -3041,13 +3004,13 @@ void FEdModeFoliage::OnInstanceCountUpdated(const UFoliageType* FoliageType)
 	UWorld* World = GetWorld();
 	ULevel* CurrentLevel = World->GetCurrentLevel();
 
-	for (FFoliageMeshInfoIterator It(World, FoliageType); It; ++It)
+	for (FFoliageInfoIterator It(World, FoliageType); It; ++It)
 	{
-		FFoliageMeshInfo* MeshInfo = (*It);
-		InstanceCountTotal += MeshInfo->Instances.Num();
+		FFoliageInfo* FoliageInfo = (*It);
+		InstanceCountTotal += FoliageInfo->Instances.Num();
 		if (It.GetActor()->GetLevel() == CurrentLevel)
 		{
-			InstanceCountCurrentLevel = MeshInfo->Instances.Num();
+			InstanceCountCurrentLevel = FoliageInfo->Instances.Num();
 		}
 	}
 
@@ -3071,10 +3034,10 @@ void FEdModeFoliage::CalcTotalInstanceCount(int32& OutInstanceCountTotal, int32&
 		if (IFA)
 		{
 			int32 IFAInstanceCount = 0;
-			for (const auto& MeshPair : IFA->FoliageMeshes)
+			for (const auto& MeshPair : IFA->FoliageInfos)
 			{
-				const FFoliageMeshInfo& MeshInfo = *MeshPair.Value;
-				IFAInstanceCount += MeshInfo.Instances.Num();
+				const FFoliageInfo& FoliageInfo = *MeshPair.Value;
+				IFAInstanceCount += FoliageInfo.Instances.Num();
 			}
 
 			OutInstanceCountTotal += IFAInstanceCount;
@@ -3182,7 +3145,7 @@ void FEdModeFoliage::MoveSelectedFoliageToLevel(ULevel* InTargetLevel)
 				bool CanMoveInstanceType = true;
 
 				// Make sure all our foliage type used by our selected instances are asset otherwise promote them to assets
-				TMap<UFoliageType*, FFoliageMeshInfo*> SelectedInstanceFoliageTypes = IFA->GetSelectedInstancesFoliageType();
+				TMap<UFoliageType*, FFoliageInfo*> SelectedInstanceFoliageTypes = IFA->GetSelectedInstancesFoliageType();
 
 				for (auto& MeshPair : SelectedInstanceFoliageTypes)
 				{
@@ -3204,8 +3167,8 @@ void FEdModeFoliage::MoveSelectedFoliageToLevel(ULevel* InTargetLevel)
 						if (NewFoliageType != nullptr)
 						{
 							// Restore previous selection for move operation
-							FFoliageMeshInfo* MeshInfo = IFA->FindMesh(NewFoliageType);
-							MeshInfo->SelectInstances(IFA, true, PreviousSelectionArray);
+							FFoliageInfo* FoliageInfo = IFA->FindInfo(NewFoliageType);
+							FoliageInfo->SelectInstances(IFA, true, PreviousSelectionArray);
 						}
 					}
 				}
@@ -3243,7 +3206,7 @@ UFoliageType* FEdModeFoliage::AddFoliageAsset(UObject* InAsset)
 			const FScopedTransaction Transaction(NSLOCTEXT("UnrealEd", "FoliageMode_AddTypeTransaction", "Add Foliage Type"));
 
 			AInstancedFoliageActor* IFA = AInstancedFoliageActor::GetInstancedFoliageActorForCurrentLevel(World, true);
-			FoliageType = IFA->GetLocalFoliageTypeForMesh(StaticMesh);
+			FoliageType = IFA->GetLocalFoliageTypeForSource(StaticMesh);
 			if (!FoliageType)
 			{
 				IFA->AddMesh(StaticMesh, &FoliageType);
@@ -3290,7 +3253,7 @@ bool FEdModeFoliage::RemoveFoliageType(UFoliageType** FoliageTypeList, int32 Num
 	for (int32 FoliageTypeIdx = 0; FoliageTypeIdx < Num; ++FoliageTypeIdx)
 	{
 		UFoliageType* FoliageType = FoliageTypeList[FoliageTypeIdx];
-		for (FFoliageMeshInfoIterator It(World, FoliageType); It; ++It)
+		for (FFoliageInfoIterator It(World, FoliageType); It; ++It)
 		{
 			IFAList.Add(It.GetActor());
 		}
@@ -3321,17 +3284,17 @@ void FEdModeFoliage::BakeFoliage(UFoliageType* Settings, bool bSelectedOnly)
 		return;
 	}
 
-	FFoliageMeshInfo* MeshInfo = IFA->FindMesh(Settings);
-	if (MeshInfo != nullptr)
+	FFoliageInfo* FoliageInfo = IFA->FindInfo(Settings);
+	if (FoliageInfo != nullptr && FoliageInfo->Type == EFoliageImplType::StaticMesh)
 	{
 		TArray<int32> InstancesToConvert;
 		if (bSelectedOnly)
 		{
-			InstancesToConvert = MeshInfo->SelectedIndices.Array();
+			InstancesToConvert = FoliageInfo->SelectedIndices.Array();
 		}
 		else
 		{
-			for (int32 InstanceIdx = 0; InstanceIdx < MeshInfo->Instances.Num(); InstanceIdx++)
+			for (int32 InstanceIdx = 0; InstanceIdx < FoliageInfo->Instances.Num(); InstanceIdx++)
 			{
 				InstancesToConvert.Add(InstanceIdx);
 			}
@@ -3340,18 +3303,18 @@ void FEdModeFoliage::BakeFoliage(UFoliageType* Settings, bool bSelectedOnly)
 		// Convert
 		for (int32 Idx = 0; Idx < InstancesToConvert.Num(); Idx++)
 		{
-			FFoliageInstance& Instance = MeshInfo->Instances[InstancesToConvert[Idx]];
+			FFoliageInstance& Instance = FoliageInfo->Instances[InstancesToConvert[Idx]];
 			// We need a world in which to spawn the actor. Use the one from the original instance.
 			UWorld* World = IFA->GetWorld();
 			check(World != nullptr);
 			AStaticMeshActor* SMA = World->SpawnActor<AStaticMeshActor>(Instance.Location, Instance.Rotation);
-			SMA->GetStaticMeshComponent()->SetStaticMesh(Settings->GetStaticMesh());
+			SMA->GetStaticMeshComponent()->SetStaticMesh(Cast<UStaticMesh>(Settings->GetSource()));
 			SMA->GetRootComponent()->SetRelativeScale3D(Instance.DrawScale3D);
 			SMA->MarkComponentsRenderStateDirty();
 		}
 
 		// Remove
-		MeshInfo->RemoveInstances(IFA, InstancesToConvert, true);
+		FoliageInfo->RemoveInstances(IFA, InstancesToConvert, true);
 	}
 }
 
@@ -3363,11 +3326,11 @@ UFoliageType* FEdModeFoliage::CopySettingsObject(UFoliageType* Settings)
 	AInstancedFoliageActor* IFA = AInstancedFoliageActor::GetInstancedFoliageActorForCurrentLevel(GetWorld());
 	IFA->Modify();
 
-	TUniqueObj<FFoliageMeshInfo> MeshInfo;
-	if (IFA->FoliageMeshes.RemoveAndCopyValue(Settings, MeshInfo))
+	TUniqueObj<FFoliageInfo> FoliageInfo;
+	if (IFA->FoliageInfos.RemoveAndCopyValue(Settings, FoliageInfo))
 	{
 		Settings = (UFoliageType*)StaticDuplicateObject(Settings, IFA, NAME_None, RF_AllFlags & ~(RF_Standalone | RF_Public));
-		IFA->FoliageMeshes.Add(Settings, MoveTemp(MeshInfo));
+		IFA->FoliageInfos.Add(Settings, MoveTemp(FoliageInfo));
 		return Settings;
 	}
 	else
@@ -3402,10 +3365,10 @@ UFoliageType* FEdModeFoliage::SaveFoliageTypeObject(UFoliageType* InFoliageType)
 void FEdModeFoliage::ReallocateClusters(UFoliageType* Settings)
 {
 	UWorld* World = GetWorld();
-	for (FFoliageMeshInfoIterator It(World, Settings); It; ++It)
+	for (FFoliageInfoIterator It(World, Settings); It; ++It)
 	{
-		FFoliageMeshInfo* MeshInfo = (*It);
-		MeshInfo->ReallocateClusters(It.GetActor(), Settings);
+		FFoliageInfo* FoliageInfo = (*It);
+		FoliageInfo->ReallocateClusters(It.GetActor(), Settings);
 	}
 }
 
@@ -3525,19 +3488,13 @@ void FEdModeFoliage::DrawHUD(FEditorViewportClient* ViewportClient, FViewport* V
 /** FEdMode: Check to see if an actor can be selected in this mode - no side effects */
 bool FEdModeFoliage::IsSelectionAllowed(AActor* InActor, bool bInSelection) const
 {
-	return false;
+	return AInstancedFoliageActor::IsOwnedByFoliage(InActor);
 }
 
 /** FEdMode: Handling SelectActor */
 bool FEdModeFoliage::Select(AActor* InActor, bool bInSelected)
 {
-	// return true if you filter that selection
-	// however - return false if we are trying to deselect so that it will infact do the deselection
-	if (bInSelected == false)
-	{
-		return false;
-	}
-	return true;
+	return false;
 }
 
 /** FEdMode: Called when the currently selected actor has changed */
@@ -3590,6 +3547,16 @@ bool FEdModeFoliage::HandleClick(FEditorViewportClient* InViewportClient, HHitPr
 				UpdateWidgetLocationToInstanceSelection();
 			}
 		}
+		else if (HitProxy && HitProxy->IsA(HActor::StaticGetType()) && AInstancedFoliageActor::IsOwnedByFoliage(((HActor*)HitProxy)->Actor))
+		{
+			HActor* ActorProxy = ((HActor*)HitProxy);
+			AInstancedFoliageActor* IFA = AInstancedFoliageActor::GetInstancedFoliageActorForLevel(ActorProxy->Actor->GetLevel());
+			if (IFA)
+			{
+				IFA->SelectInstance(ActorProxy->Actor, Click.IsControlDown());
+				UpdateWidgetLocationToInstanceSelection();
+			}
+		}
 		else
 		{
 			if (!Click.IsControlDown())
@@ -3679,6 +3646,8 @@ bool FEdModeFoliage::InputDelta(FEditorViewportClient* InViewportClient, FViewpo
 		{
 			static const float RadiusAdjustmentFactor = 10.f;
 			AdjustBrushRadius(RadiusAdjustmentFactor * InDrag.Y);
+
+			return true;
 		}
 	}
 	else
@@ -3691,6 +3660,8 @@ bool FEdModeFoliage::InputDelta(FEditorViewportClient* InViewportClient, FViewpo
 
 			// Only allow alt-drag on first InputDelta
 			bCanAltDrag = false;
+
+			return true;
 		}
 	}
 
diff --git a/Engine/Source/Editor/FoliageEdit/Private/FoliageEdMode.h b/Engine/Source/Editor/FoliageEdit/Private/FoliageEdMode.h
--- a/Engine/Source/Editor/FoliageEdit/Private/FoliageEdMode.h
+++ b/Engine/Source/Editor/FoliageEdit/Private/FoliageEdMode.h
@@ -193,7 +193,7 @@ class FMeshInfoSnapshot
 	FFoliageInstanceHash Hash;
 	TArray<FVector> Locations;
 public:
-	FMeshInfoSnapshot(FFoliageMeshInfo* MeshInfo)
+	FMeshInfoSnapshot(FFoliageInfo* MeshInfo)
 		: Hash(*MeshInfo->InstanceHash)
 	{
 		int32 NumInstances = MeshInfo->Instances.Num();
@@ -302,6 +302,9 @@ public:
 	/** Called when asset is removed */
 	void NotifyAssetRemoved(const FAssetData& AssetInfo);
 
+	/** Called when actor foliage selection changes */
+	void NotifyActorSelectionChanged(bool bSelect, const TArray<AActor*>& Selection);
+
 	/**
 	* Called when the mouse is moved over the viewport
 	*
@@ -527,7 +530,7 @@ private:
 	/** Reapply instance settings to exiting instances */
 	void ReapplyInstancesDensityForBrush(UWorld* InWorld, const UFoliageType* Settings, const FSphere& BrushSphere, float Pressure);
 	void ReapplyInstancesForBrush(UWorld* InWorld, const UFoliageType* Settings, const FSphere& BrushSphere, float Pressure);
-	void ReapplyInstancesForBrush(UWorld* InWorld, AInstancedFoliageActor* IFA, const UFoliageType* Settings, FFoliageMeshInfo* MeshInfo, const FSphere& BrushSphere, float Pressure);
+	void ReapplyInstancesForBrush(UWorld* InWorld, AInstancedFoliageActor* IFA, const UFoliageType* Settings, FFoliageInfo* MeshInfo, const FSphere& BrushSphere, float Pressure);
 
 	/** Select instances inside the brush. */
 	void SelectInstancesForBrush(UWorld* InWorld, const UFoliageType* Settings, const FSphere& BrushSphere, bool bSelect);
@@ -541,8 +544,8 @@ private:
 	/** Set/Clear selection for foliage instances of specific type  */
 	void SelectInstances(UWorld* InWorld, const UFoliageType* Settings, bool bSelect);
 
-	/**  Propagate the selected foliage instances to the actual render components */
-	void ApplySelectionToComponents(UWorld* InWorld, bool bApply);
+	/**  Propagate the selected foliage instances to the actual render foliage */
+	void ApplySelection(UWorld* InWorld, bool bApply);
 
 	/**  Applies relative transformation to selected instances */
 	void TransformSelectedInstances(UWorld* InWorld, const FVector& InDrag, const FRotator& InRot, const FVector& InScale, bool bDuplicate);
@@ -557,7 +560,7 @@ private:
 	void RemoveSelectedInstances(UWorld* InWorld);
 			
 	/** Snap instance to the ground   */
-	bool SnapInstanceToGround(AInstancedFoliageActor* InIFA, float AlignMaxAngle, FFoliageMeshInfo& Mesh, int32 InstanceIdx);
+	bool SnapInstanceToGround(AInstancedFoliageActor* InIFA, float AlignMaxAngle, FFoliageInfo& Mesh, int32 InstanceIdx);
 	void SnapSelectedInstancesToGround(UWorld* InWorld);
 
 	/** Callback for when an actor is spawned (to check if it's a new IFA) */
diff --git a/Engine/Source/Editor/FoliageEdit/Private/FoliageEditModule.cpp b/Engine/Source/Editor/FoliageEdit/Private/FoliageEditModule.cpp
--- a/Engine/Source/Editor/FoliageEdit/Private/FoliageEditModule.cpp
+++ b/Engine/Source/Editor/FoliageEdit/Private/FoliageEditModule.cpp
@@ -16,6 +16,7 @@ const FName FoliageEditAppIdentifier = FName(TEXT("FoliageEdApp"));
 
 #include "UObject/UObjectIterator.h"
 #include "FoliageType_InstancedStaticMesh.h"
+#include "FoliageType_Actor.h"
 #include "InstancedFoliageActor.h"
 #include "FoliageEdMode.h"
 #include "PropertyEditorModule.h"
@@ -28,6 +29,7 @@ const FName FoliageEditAppIdentifier = FName(TEXT("FoliageEdApp"));
 #include "ProceduralFoliageBlockingVolume.h"
 #include "FoliageTypeObjectCustomization.h"
 #include "FoliageType_ISMThumbnailRenderer.h"
+#include "FoliageType_ActorThumbnailRenderer.h"
 #include "EditorModeManager.h"
 #include "LevelEditorViewport.h"
 
@@ -77,6 +79,7 @@ public:
 
 		// Register thumbnail renderer
 		UThumbnailManager::Get().RegisterCustomRenderer(UFoliageType_InstancedStaticMesh::StaticClass(), UFoliageType_ISMThumbnailRenderer::StaticClass());
+		UThumbnailManager::Get().RegisterCustomRenderer(UFoliageType_Actor::StaticClass(), UFoliageType_ActorThumbnailRenderer::StaticClass());
 	}
 
 	/**
diff --git a/Engine/Source/Editor/FoliageEdit/Private/FoliageEditUtility.cpp b/Engine/Source/Editor/FoliageEdit/Private/FoliageEditUtility.cpp
--- a/Engine/Source/Editor/FoliageEdit/Private/FoliageEditUtility.cpp
+++ b/Engine/Source/Editor/FoliageEdit/Private/FoliageEditUtility.cpp
@@ -22,11 +22,11 @@ UFoliageType* FFoliageEditUtility::SaveFoliageTypeObject(UFoliageType* InFoliage
 	if (!InFoliageType->IsAsset())
 	{
 		FString PackageName;
-		UStaticMesh* StaticMesh = InFoliageType->GetStaticMesh();
-		if (StaticMesh)
+		UObject* FoliageSource = InFoliageType->GetSource();
+		if (FoliageSource)
 		{
 			// Build default settings asset name and path
-			PackageName = FPackageName::GetLongPackagePath(StaticMesh->GetOutermost()->GetName()) + TEXT("/") + StaticMesh->GetName() + TEXT("_FoliageType");
+			PackageName = FPackageName::GetLongPackagePath(FoliageSource->GetOutermost()->GetName()) + TEXT("/") + FoliageSource->GetName() + TEXT("_FoliageType");
 		}
 
 		TSharedRef<SDlgPickAssetPath> SaveFoliageTypeDialog =
@@ -91,35 +91,26 @@ void FFoliageEditUtility::ReplaceFoliageTypeObject(UWorld* InWorld, UFoliageType
 			if (IFA)
 			{
 				IFA->Modify();
-				TUniqueObj<FFoliageMeshInfo> OldMeshInfo;
-				IFA->FoliageMeshes.RemoveAndCopyValue(OldType, OldMeshInfo);
+				TUniqueObj<FFoliageInfo> OldInfo;
+				IFA->FoliageInfos.RemoveAndCopyValue(OldType, OldInfo);
 
 				// Old component needs to go
-				if (OldMeshInfo->Component != nullptr)
+				if (OldInfo->IsInitialized())
 				{
-					if (OldMeshInfo->Component->GetStaticMesh() != nullptr)
-					{
-						OldMeshInfo->Component->GetStaticMesh()->GetOnExtendedBoundsChanged().RemoveAll(&(*OldMeshInfo));
-					}
-
-					OldMeshInfo->Component->ClearInstances();
-					OldMeshInfo->Component->SetFlags(RF_Transactional);
-					OldMeshInfo->Component->Modify();
-					OldMeshInfo->Component->DestroyComponent();
-					OldMeshInfo->Component = nullptr;
+					OldInfo->Uninitialize();
 				}
-
+				
 				// Append instances if new foliage type is already exists in this actor
 				// Otherwise just replace key entry for instances
-				TUniqueObj<FFoliageMeshInfo>* NewMeshInfo = IFA->FoliageMeshes.Find(NewType);
-				if (NewMeshInfo)
+				TUniqueObj<FFoliageInfo>* NewInfo = IFA->FoliageInfos.Find(NewType);
+				if (NewInfo)
 				{
-					(*NewMeshInfo)->Instances.Append(OldMeshInfo->Instances);
-					(*NewMeshInfo)->ReallocateClusters(IFA, NewType);
+					(*NewInfo)->Instances.Append(OldInfo->Instances);
+					(*NewInfo)->ReallocateClusters(IFA, NewType);
 				}
 				else
 				{
-					IFA->FoliageMeshes.Add(NewType, MoveTemp(OldMeshInfo))->ReallocateClusters(IFA, NewType);
+					IFA->FoliageInfos.Add(NewType, MoveTemp(OldInfo))->ReallocateClusters(IFA, NewType);
 				}
 			}
 		}
@@ -161,7 +152,7 @@ void FFoliageEditUtility::MoveActorFoliageInstancesToLevel(ULevel* InTargetLevel
 
 			bool CanMoveInstanceType = true;
 
-			TMap<UFoliageType*, FFoliageMeshInfo*> InstancesFoliageType = IFA->GetAllInstancesFoliageType();
+			TMap<UFoliageType*, FFoliageInfo*> InstancesFoliageType = IFA->GetAllInstancesFoliageType();
 
 			for (auto& MeshPair : InstancesFoliageType)
 			{
@@ -189,7 +180,7 @@ void FFoliageEditUtility::MoveActorFoliageInstancesToLevel(ULevel* InTargetLevel
 					if (NewFoliageType != nullptr)
 					{
 						// Restore previous selection for move operation
-						FFoliageMeshInfo* MeshInfo = IFA->FindMesh(NewFoliageType);
+						FFoliageInfo* MeshInfo = IFA->FindInfo(NewFoliageType);
 						MeshInfo->SelectInstances(IFA, true, PreviousSelectionArray);
 					}
 				}
diff --git a/Engine/Source/Editor/FoliageEdit/Private/FoliagePaletteItem.cpp b/Engine/Source/Editor/FoliageEdit/Private/FoliagePaletteItem.cpp
--- a/Engine/Source/Editor/FoliageEdit/Private/FoliagePaletteItem.cpp
+++ b/Engine/Source/Editor/FoliageEdit/Private/FoliagePaletteItem.cpp
@@ -44,7 +44,7 @@ FFoliagePaletteItemModel::FFoliagePaletteItemModel(FFoliageMeshUIInfoPtr InTypeI
 	}
 	else
 	{
-		AssetData = FAssetData(TypeInfo->Settings->GetClass()->GetDefaultObject<UFoliageType>()->GetStaticMesh());
+		AssetData = FAssetData(TypeInfo->Settings->GetClass()->GetDefaultObject<UFoliageType>()->GetSource());
 	}
 	
 
@@ -56,7 +56,7 @@ FFoliagePaletteItemModel::FFoliagePaletteItemModel(FFoliageMeshUIInfoPtr InTypeI
 	if (TypeInfo->Settings->IsAsset())
 	{
 		// Get the Foliage Type asset color
-		auto FoliageTypeAssetActions =  AssetToolsModule.Get().GetAssetTypeActionsForClass(UFoliageType_InstancedStaticMesh::StaticClass());
+		auto FoliageTypeAssetActions =  AssetToolsModule.Get().GetAssetTypeActionsForClass(TypeInfo->Settings->GetClass());
 		if (FoliageTypeAssetActions.IsValid())
 		{
 			ThumbnailConfig.AssetTypeColorOverride = FoliageTypeAssetActions.Pin()->GetTypeColor();
diff --git a/Engine/Source/Editor/FoliageEdit/Private/FoliageTypeFactory.cpp b/Engine/Source/Editor/FoliageEdit/Private/FoliageTypeFactory.cpp
--- a/Engine/Source/Editor/FoliageEdit/Private/FoliageTypeFactory.cpp
+++ b/Engine/Source/Editor/FoliageEdit/Private/FoliageTypeFactory.cpp
@@ -3,8 +3,11 @@
 #include "FoliageTypeFactory.h"
 #include "AssetTypeCategories.h"
 #include "FoliageType_InstancedStaticMesh.h"
+#include "FoliageType_Actor.h"
 
-UFoliageTypeFactory::UFoliageTypeFactory(const FObjectInitializer& ObjectInitializer)
+#define LOCTEXT_NAMESPACE "FoliageTypeFactory"
+
+UFoliageType_InstancedStaticMeshFactory::UFoliageType_InstancedStaticMeshFactory(const FObjectInitializer& ObjectInitializer)
 	: Super(ObjectInitializer)
 {
 	bCreateNew = true;
@@ -12,14 +15,32 @@ UFoliageTypeFactory::UFoliageTypeFactory(const FObjectInitializer& ObjectInitial
 	SupportedClass = UFoliageType_InstancedStaticMesh::StaticClass();
 }
 
-UObject* UFoliageTypeFactory::FactoryCreateNew(UClass* Class, UObject* InParent, FName Name, EObjectFlags Flags, UObject* Context, FFeedbackContext* Warn)
+UObject* UFoliageType_InstancedStaticMeshFactory::FactoryCreateNew(UClass* Class, UObject* InParent, FName Name, EObjectFlags Flags, UObject* Context, FFeedbackContext* Warn)
 {
-	auto NewFoliageType = NewObject<UFoliageType_InstancedStaticMesh>(InParent, Class, Name, Flags | RF_Transactional);
+	return NewObject<UFoliageType_InstancedStaticMesh>(InParent, Class, Name, Flags | RF_Transactional);
+}
 
-	return NewFoliageType;
+FText UFoliageType_InstancedStaticMeshFactory::GetToolTip() const
+{
+	return LOCTEXT("FoliageTypeStaticMeshToolTip", "Static Mesh Foliage is a foliage type that will use mesh instancing and is optimal for non-interactive foliage.");
 }
 
-uint32 UFoliageTypeFactory::GetMenuCategories() const
+UFoliageType_ActorFactory::UFoliageType_ActorFactory(const FObjectInitializer& ObjectInitializer)
+	: Super(ObjectInitializer)
 {
-	return EAssetTypeCategories::Misc;
+	bCreateNew = true;
+	bEditAfterNew = true;
+	SupportedClass = UFoliageType_Actor::StaticClass();
 }
+
+UObject* UFoliageType_ActorFactory::FactoryCreateNew(UClass* Class, UObject* InParent, FName Name, EObjectFlags Flags, UObject* Context, FFeedbackContext* Warn)
+{
+	return NewObject<UFoliageType_Actor>(InParent, Class, Name, Flags | RF_Transactional);
+}
+
+FText UFoliageType_ActorFactory::GetToolTip() const
+{
+	return LOCTEXT("FoliageTypeActorToolTip", "Blueprint Foliage is a foliage type that will place blueprint actor instances. The cost of painting this foliage is the same as adding actors in a scene.");
+}
+
+#undef LOCTEXT_NAMESPACE
diff --git a/Engine/Source/Editor/FoliageEdit/Private/FoliageTypeFactory.h b/Engine/Source/Editor/FoliageEdit/Private/FoliageTypeFactory.h
--- a/Engine/Source/Editor/FoliageEdit/Private/FoliageTypeFactory.h
+++ b/Engine/Source/Editor/FoliageEdit/Private/FoliageTypeFactory.h
@@ -12,12 +12,31 @@
 #include "FoliageTypeFactory.generated.h"
 
 UCLASS()
-class UFoliageTypeFactory : public UFactory
+class UFoliageType_InstancedStaticMeshFactory : public UFactory
 {
 	GENERATED_UCLASS_BODY()
 
 	// UFactory interface
 	virtual UObject* FactoryCreateNew(UClass* Class, UObject* InParent, FName Name, EObjectFlags Flags, UObject* Context, FFeedbackContext* Warn) override;
-	virtual uint32 GetMenuCategories() const;
+	virtual FString GetDefaultNewAssetName() const override
+	{
+		return TEXT("NewInstancedStaticMeshFoliage");
+	}
+	virtual FText GetToolTip() const override;
 	// End of UFactory interface
 };
+
+UCLASS()
+class UFoliageType_ActorFactory : public UFactory
+{
+	GENERATED_UCLASS_BODY()
+
+	// UFactory interface
+	virtual UObject* FactoryCreateNew(UClass* Class, UObject* InParent, FName Name, EObjectFlags Flags, UObject* Context, FFeedbackContext* Warn) override;
+	virtual FString GetDefaultNewAssetName() const override
+	{
+		return TEXT("NewBlueprintFoliage");
+	}
+	virtual FText GetToolTip() const override;
+	// End of UFactory interface
+};
\ No newline at end of file
diff --git a/Engine/Source/Editor/FoliageEdit/Private/FoliageTypeObjectCustomization.cpp b/Engine/Source/Editor/FoliageEdit/Private/FoliageTypeObjectCustomization.cpp
--- a/Engine/Source/Editor/FoliageEdit/Private/FoliageTypeObjectCustomization.cpp
+++ b/Engine/Source/Editor/FoliageEdit/Private/FoliageTypeObjectCustomization.cpp
@@ -4,6 +4,7 @@
 #include "Widgets/DeclarativeSyntaxSupport.h"
 #include "AssetData.h"
 #include "FoliageType_InstancedStaticMesh.h"
+#include "FoliageType_Actor.h"
 #include "PropertyHandle.h"
 #include "IDetailChildrenBuilder.h"
 #include "PropertyCustomizationHelpers.h"
@@ -15,6 +16,7 @@ void FFoliageTypeObjectCustomization::CustomizeHeader(TSharedRef<IPropertyHandle
 	// Only allow foliage type assets to be created (i.e. don't show all the factories for blueprints)
 	TArray<const UClass*> SupportedClasses;
 	SupportedClasses.Add(UFoliageType_InstancedStaticMesh::StaticClass());
+	SupportedClasses.Add(UFoliageType_Actor::StaticClass());
 
 	HeaderRow
 		.NameContent()
@@ -43,7 +45,7 @@ bool FFoliageTypeObjectCustomization::OnShouldFilterAsset(const FAssetData& Asse
 {
 	// If the asset is a BP class that doesn't inherit from UFoliageType, hide it
 	const FString ParentClassName = AssetData.GetTagValueRef<FString>(FBlueprintTags::ParentClassPath);
-	if (!ParentClassName.IsEmpty() && !ParentClassName.Contains(TEXT("FoliageType_InstancedStaticMesh")))
+	if (!ParentClassName.IsEmpty() && !ParentClassName.Contains(TEXT("FoliageType_InstancedStaticMesh")) && !ParentClassName.Contains(TEXT("FoliageType_Actor")))
 	{
 		return true;
 	}
diff --git a/Engine/Source/Editor/FoliageEdit/Private/FoliageType_ActorThumbnailRenderer.cpp b/Engine/Source/Editor/FoliageEdit/Private/FoliageType_ActorThumbnailRenderer.cpp
new file mode 100644
--- /dev/null
+++ b/Engine/Source/Editor/FoliageEdit/Private/FoliageType_ActorThumbnailRenderer.cpp
@@ -0,0 +1,31 @@
+// Copyright 1998-2019 Epic Games, Inc. All Rights Reserved.
+
+#include "FoliageType_ActorThumbnailRenderer.h"
+#include "FoliageType_Actor.h"
+#include "ShowFlags.h"
+#include "SceneView.h"
+#include "Misc/App.h"
+
+bool UFoliageType_ActorThumbnailRenderer::CanVisualizeAsset(UObject* Object)
+{
+	UFoliageType_Actor* FoliageType = Cast<UFoliageType_Actor>(Object);
+	if (FoliageType == nullptr || FoliageType->ActorClass == nullptr)
+	{
+		return false;
+	}
+
+	return Super::CanVisualizeAsset(FoliageType->ActorClass->ClassGeneratedBy);
+}
+
+void UFoliageType_ActorThumbnailRenderer::Draw(UObject* Object, int32 X, int32 Y, uint32 Width, uint32 Height, FRenderTarget* RenderTarget, FCanvas* Canvas)
+{
+	UFoliageType_Actor* FoliageType = Cast<UFoliageType_Actor>(Object);
+	if (FoliageType == nullptr || FoliageType->ActorClass == nullptr)
+	{
+		return;
+	}
+
+	Super::Draw(FoliageType->ActorClass->ClassGeneratedBy, X, Y, Width, Height, RenderTarget, Canvas);
+}
+
+
diff --git a/Engine/Source/Editor/FoliageEdit/Private/FoliageType_InstancedStaticMeshPaintingCustomization.cpp b/Engine/Source/Editor/FoliageEdit/Private/FoliageType_InstancedStaticMeshPaintingCustomization.cpp
--- a/Engine/Source/Editor/FoliageEdit/Private/FoliageType_InstancedStaticMeshPaintingCustomization.cpp
+++ b/Engine/Source/Editor/FoliageEdit/Private/FoliageType_InstancedStaticMeshPaintingCustomization.cpp
@@ -59,7 +59,7 @@ void FFoliageType_InstancedStaticMeshPaintingCustomization::CustomizeDetails(IDe
 			UFoliageType* FoliageType = TypeInfo->Settings;
 			if (!FoliageType->IsAsset() && !FoliageType->GetClass()->ClassGeneratedBy)
 			{
-				UnavailableMeshNames.Add(FoliageType->GetStaticMesh()->GetFName());
+				UnavailableMeshNames.Add(FoliageType->GetSource()->GetFName());
 			}
 		}
 
diff --git a/Engine/Source/Editor/FoliageEdit/Private/LandscapeGrassTypeFactory.cpp b/Engine/Source/Editor/FoliageEdit/Private/LandscapeGrassTypeFactory.cpp
--- a/Engine/Source/Editor/FoliageEdit/Private/LandscapeGrassTypeFactory.cpp
+++ b/Engine/Source/Editor/FoliageEdit/Private/LandscapeGrassTypeFactory.cpp
@@ -18,8 +18,3 @@ UObject* ULandscapeGrassTypeFactory::FactoryCreateNew(UClass* Class, UObject* In
 
 	return NewGrassType;
 }
-
-uint32 ULandscapeGrassTypeFactory::GetMenuCategories() const
-{
-	return EAssetTypeCategories::Misc;
-}
diff --git a/Engine/Source/Editor/FoliageEdit/Private/LandscapeGrassTypeFactory.h b/Engine/Source/Editor/FoliageEdit/Private/LandscapeGrassTypeFactory.h
--- a/Engine/Source/Editor/FoliageEdit/Private/LandscapeGrassTypeFactory.h
+++ b/Engine/Source/Editor/FoliageEdit/Private/LandscapeGrassTypeFactory.h
@@ -18,6 +18,5 @@ class ULandscapeGrassTypeFactory : public UFactory
 
 	// UFactory interface
 	virtual UObject* FactoryCreateNew(UClass* Class, UObject* InParent, FName Name, EObjectFlags Flags, UObject* Context, FFeedbackContext* Warn) override;
-	virtual uint32 GetMenuCategories() const override;
 	// End of UFactory interface
 };
diff --git a/Engine/Source/Editor/FoliageEdit/Private/ProceduralFoliageSpawnerFactory.cpp b/Engine/Source/Editor/FoliageEdit/Private/ProceduralFoliageSpawnerFactory.cpp
--- a/Engine/Source/Editor/FoliageEdit/Private/ProceduralFoliageSpawnerFactory.cpp
+++ b/Engine/Source/Editor/FoliageEdit/Private/ProceduralFoliageSpawnerFactory.cpp
@@ -20,11 +20,6 @@ UObject* UProceduralFoliageSpawnerFactory::FactoryCreateNew(UClass* Class, UObje
 	return NewProceduralFoliage;
 }
 
-uint32 UProceduralFoliageSpawnerFactory::GetMenuCategories() const
-{
-	return EAssetTypeCategories::Misc;
-}
-
 bool UProceduralFoliageSpawnerFactory::ShouldShowInNewMenu() const
 {
 	return GetDefault<UEditorExperimentalSettings>()->bProceduralFoliage;
diff --git a/Engine/Source/Editor/FoliageEdit/Private/ProceduralFoliageSpawnerFactory.h b/Engine/Source/Editor/FoliageEdit/Private/ProceduralFoliageSpawnerFactory.h
--- a/Engine/Source/Editor/FoliageEdit/Private/ProceduralFoliageSpawnerFactory.h
+++ b/Engine/Source/Editor/FoliageEdit/Private/ProceduralFoliageSpawnerFactory.h
@@ -18,7 +18,6 @@ class UProceduralFoliageSpawnerFactory : public UFactory
 
 	// UFactory interface
 	virtual UObject* FactoryCreateNew(UClass* Class, UObject* InParent, FName Name, EObjectFlags Flags, UObject* Context, FFeedbackContext* Warn) override;
-	virtual uint32 GetMenuCategories() const override;
 	virtual bool ShouldShowInNewMenu() const override;
 	// End of UFactory interface
 };
diff --git a/Engine/Source/Editor/FoliageEdit/Private/SFoliagePalette.cpp b/Engine/Source/Editor/FoliageEdit/Private/SFoliagePalette.cpp
--- a/Engine/Source/Editor/FoliageEdit/Private/SFoliagePalette.cpp
+++ b/Engine/Source/Editor/FoliageEdit/Private/SFoliagePalette.cpp
@@ -675,7 +675,7 @@ void SFoliagePalette::OnSearchTextChanged(const FText& InFilterText)
 TSharedRef<SWidget> SFoliagePalette::GetAddFoliageTypePicker()
 {
 	TArray<const UClass*> ClassFilters;
-	ClassFilters.Add(UFoliageType_InstancedStaticMesh::StaticClass());
+	ClassFilters.Add(UFoliageType::StaticClass());
 
 	return PropertyCustomizationHelpers::MakeAssetPickerWithMenu(FAssetData(),
 		false,
@@ -1129,7 +1129,7 @@ void SFoliagePalette::OnShowFoliageTypeInCB()
 		}
 		else
 		{
-			SelectedAssets.Add(FoliageType->GetStaticMesh());
+			SelectedAssets.Add(FoliageType->GetSource());
 		}
 	}
 
diff --git a/Engine/Source/Editor/FoliageEdit/Public/FoliageType_ActorThumbnailRenderer.h b/Engine/Source/Editor/FoliageEdit/Public/FoliageType_ActorThumbnailRenderer.h
new file mode 100644
--- /dev/null
+++ b/Engine/Source/Editor/FoliageEdit/Public/FoliageType_ActorThumbnailRenderer.h
@@ -0,0 +1,30 @@
+// Copyright 1998-2019 Epic Games, Inc. All Rights Reserved.
+
+/**
+ * This thumbnail renderer displays the actor used by this foliage type
+ */
+
+#pragma once
+
+#include "CoreMinimal.h"
+#include "UObject/ObjectMacros.h"
+#include "ThumbnailRendering/BlueprintThumbnailRenderer.h"
+#include "FoliageType_ActorThumbnailRenderer.generated.h"
+
+
+class FCanvas;
+class FRenderTarget;
+
+UCLASS(CustomConstructor, Config = Editor)
+class UFoliageType_ActorThumbnailRenderer : public UBlueprintThumbnailRenderer
+{
+	GENERATED_UCLASS_BODY()
+
+	UFoliageType_ActorThumbnailRenderer(const FObjectInitializer& ObjectInitializer)
+		: Super(ObjectInitializer)
+	{}
+
+	// UThumbnailRenderer implementation
+	virtual void Draw(UObject* Object, int32 X, int32 Y, uint32 Width, uint32 Height, FRenderTarget*, FCanvas* Canvas) override;
+	virtual bool CanVisualizeAsset(UObject* Object) override;
+};
diff --git a/Engine/Source/Editor/UnrealEd/Classes/ThumbnailRendering/ClassThumbnailRenderer.h b/Engine/Source/Editor/UnrealEd/Classes/ThumbnailRendering/ClassThumbnailRenderer.h
--- a/Engine/Source/Editor/UnrealEd/Classes/ThumbnailRendering/ClassThumbnailRenderer.h
+++ b/Engine/Source/Editor/UnrealEd/Classes/ThumbnailRendering/ClassThumbnailRenderer.h
@@ -17,7 +17,7 @@ class UClassThumbnailRenderer : public UDefaultSizedThumbnailRenderer
 	GENERATED_UCLASS_BODY()
 
 	// Begin UThumbnailRenderer Object
-	virtual bool CanVisualizeAsset(UObject* Object) override;
+	UNREALED_API virtual bool CanVisualizeAsset(UObject* Object) override;
 	UNREALED_API virtual void Draw(UObject* Object, int32 X, int32 Y, uint32 Width, uint32 Height, FRenderTarget*, FCanvas* Canvas) override;
 	// End UThumbnailRenderer Object
 
diff --git a/Engine/Source/Editor/UnrealEd/Private/EditorObject.cpp b/Engine/Source/Editor/UnrealEd/Private/EditorObject.cpp
--- a/Engine/Source/Editor/UnrealEd/Private/EditorObject.cpp
+++ b/Engine/Source/Editor/UnrealEd/Private/EditorObject.cpp
@@ -334,7 +334,7 @@ static const TCHAR* ImportProperties(
 				{
 					AInstancedFoliageActor* IFA = AInstancedFoliageActor::GetInstancedFoliageActorForLevel(ActorComponent->GetComponentLevel(), true);
 
-					FFoliageMeshInfo* MeshInfo = nullptr;
+					FFoliageInfo* MeshInfo = nullptr;
 					UFoliageType* FoliageType = IFA->AddFoliageType(SourceFoliageType, &MeshInfo);
 
 					const TCHAR* StrPtr;
@@ -369,7 +369,12 @@ static const TCHAR* ImportProperties(
 						FParse::Value(StrPtr, TEXT("Flags="), Instance.Flags);
 
 						// Add the instance
-						MeshInfo->AddInstance(IFA, FoliageType, Instance, ActorComponent, true);
+						MeshInfo->AddInstance(IFA, FoliageType, Instance, ActorComponent);
+					}
+
+					if (MeshInfo)
+					{
+						MeshInfo->Refresh(IFA, true, true);
 					}
 				}
 			}
diff --git a/Engine/Source/Editor/UnrealEd/Private/EditorServer.cpp b/Engine/Source/Editor/UnrealEd/Private/EditorServer.cpp
--- a/Engine/Source/Editor/UnrealEd/Private/EditorServer.cpp
+++ b/Engine/Source/Editor/UnrealEd/Private/EditorServer.cpp
@@ -3092,7 +3092,7 @@ TArray<UFoliageType*> UEditorEngine::GetFoliageTypesInWorld(UWorld* InWorld)
 	// Iterate over all foliage actors in the world
 	for (TActorIterator<AInstancedFoliageActor> It(InWorld); It; ++It)
 	{
-		for (const auto& Pair : It->FoliageMeshes)
+		for (const auto& Pair : It->FoliageInfos)
 		{
 			FoliageSet.Add(Pair.Key);
 		}
diff --git a/Engine/Source/Runtime/Engine/Classes/Engine/World.h b/Engine/Source/Runtime/Engine/Classes/Engine/World.h
--- a/Engine/Source/Runtime/Engine/Classes/Engine/World.h
+++ b/Engine/Source/Runtime/Engine/Classes/Engine/World.h
@@ -466,6 +466,9 @@ public:
 #if WITH_EDITOR
 	/** Determines whether the begin play cycle will run on the spawned actor when in the editor. */
 	uint16 bTemporaryEditorActor:1;
+
+	/* Determines wether or not the actor should be hidden from the Scene Outliner */
+	uint16 bHideFromSceneOutliner : 1;
 #endif
 	
 	/* Flags used to describe the spawned actor/object instance. */
diff --git a/Engine/Source/Runtime/Engine/Classes/GameFramework/Actor.h b/Engine/Source/Runtime/Engine/Classes/GameFramework/Actor.h
--- a/Engine/Source/Runtime/Engine/Classes/GameFramework/Actor.h
+++ b/Engine/Source/Runtime/Engine/Classes/GameFramework/Actor.h
@@ -2989,6 +2989,9 @@ private:
 	friend struct FMarkActorIsBeingDestroyed;
 	friend struct FActorParentComponentSetter;
 	friend struct FSetActorWantsDestroyDuringBeginPlay;
+#if WITH_EDITOR
+	friend struct FSetActorHiddenInSceneOutliner;
+#endif
 
 	// Static helpers for accessing functions on SceneComponent.
 	// These are templates for no other reason than to delay compilation until USceneComponent is defined.
@@ -3067,6 +3070,19 @@ private:
 	friend UWorld;
 };
 
+#if WITH_EDITOR
+struct FSetActorHiddenInSceneOutliner
+{
+private:
+	FSetActorHiddenInSceneOutliner(AActor* InActor)
+	{
+		InActor->bListedInSceneOutliner = false;
+	}
+
+	friend UWorld;
+};
+#endif
+
 /**
  * TInlineComponentArray is simply a TArray that reserves a fixed amount of space on the stack
  * to try to avoid heap allocation when there are fewer than a specified number of elements expected in the result.
diff --git a/Engine/Source/Runtime/Engine/Private/LevelActor.cpp b/Engine/Source/Runtime/Engine/Private/LevelActor.cpp
--- a/Engine/Source/Runtime/Engine/Private/LevelActor.cpp
+++ b/Engine/Source/Runtime/Engine/Private/LevelActor.cpp
@@ -472,6 +472,10 @@ AActor* UWorld::SpawnActor( UClass* Class, FTransform const* UserTransformPtr, c
 	Actor->SpawnCollisionHandlingMethod = CollisionHandlingMethod;
 
 #if WITH_EDITOR
+	if (SpawnParameters.bHideFromSceneOutliner)
+	{
+		FSetActorHiddenInSceneOutliner SetActorHidden(Actor);
+	}
 	Actor->bIsEditorPreviewActor = SpawnParameters.bTemporaryEditorActor;
 #endif //WITH_EDITOR
 
diff --git a/Engine/Source/Runtime/Engine/Private/World.cpp b/Engine/Source/Runtime/Engine/Private/World.cpp
--- a/Engine/Source/Runtime/Engine/Private/World.cpp
+++ b/Engine/Source/Runtime/Engine/Private/World.cpp
@@ -162,6 +162,7 @@ FActorSpawnParameters::FActorSpawnParameters()
 , bAllowDuringConstructionScript(false)
 #if WITH_EDITOR
 , bTemporaryEditorActor(false)
+, bHideFromSceneOutliner(false)
 #endif
 , ObjectFlags(RF_Transactional)
 {
diff --git a/Engine/Source/Runtime/Foliage/Private/FoliageActor.cpp b/Engine/Source/Runtime/Foliage/Private/FoliageActor.cpp
new file mode 100644
--- /dev/null
+++ b/Engine/Source/Runtime/Foliage/Private/FoliageActor.cpp
@@ -0,0 +1,256 @@
+// Copyright 1998-2019 Epic Games, Inc. All Rights Reserved.
+
+#include "FoliageActor.h"
+#include "InstancedFoliageActor.h"
+#include "FoliageType_Actor.h"
+
+FName FoliageActorTag(TEXT("FoliageActorInstance"));
+
+//
+//
+// FFoliageActor
+void FFoliageActor::AddReferencedObjects(UObject* InThis, FReferenceCollector& Collector)
+{
+	for (const TWeakObjectPtr<AActor>& ActorPtr : ActorInstances)
+	{
+		if (AActor* Actor = ActorPtr.Get(false))
+		{
+			Collector.AddReferencedObject(Actor, InThis);
+		}
+	}
+}
+
+void FFoliageActor::Serialize(FArchive& Ar)
+{
+	Ar << ActorInstances;
+	Ar << ActorClass;
+}
+
+void FFoliageActor::DestroyActors(bool bOnLoad)
+{
+	for (const TWeakObjectPtr<AActor>& ActorPtr : ActorInstances)
+	{
+		if (AActor* Actor = ActorPtr.Get(false))
+		{
+			if (bOnLoad)
+			{
+				Actor->ConditionalPostLoad();
+			}
+			Actor->GetWorld()->DestroyActor(Actor);
+		}
+	}
+}
+
+#if WITH_EDITOR
+bool FFoliageActor::IsInitialized() const
+{
+	return bInitialized;
+}
+
+void FFoliageActor::Initialize(AInstancedFoliageActor* IFA, const UFoliageType* FoliageType)
+{
+	check(!bInitialized);
+	const UFoliageType_Actor* FoliageType_Actor = Cast<UFoliageType_Actor>(FoliageType);
+	ActorClass = FoliageType_Actor->ActorClass;
+	bInitialized = true;
+}
+
+void FFoliageActor::Uninitialize()
+{
+	check(bInitialized);
+	DestroyActors(false);
+	ActorInstances.Empty();
+	ActorClass = nullptr;
+	bInitialized = false;
+}
+
+AActor* FFoliageActor::Spawn(AInstancedFoliageActor* IFA, const FFoliageInstance& Instance)
+{
+	if (ActorClass == nullptr)
+	{
+		return nullptr;
+	}
+
+	FActorSpawnParameters SpawnParameters;
+	SpawnParameters.ObjectFlags = RF_Transactional;
+	SpawnParameters.bHideFromSceneOutliner = true;
+	SpawnParameters.OverrideLevel = IFA->GetLevel();
+	AActor* NewActor = IFA->GetWorld()->SpawnActor<AActor>(ActorClass, Instance.GetInstanceWorldTransform(), SpawnParameters);
+	if (NewActor)
+	{
+		NewActor->Tags.AddUnique(FoliageActorTag);
+	}
+	return NewActor;
+}
+
+TArray<AActor*> FFoliageActor::GetActorsFromSelectedIndices(const TSet<int32>& SelectedIndices) const
+{
+	TArray<AActor*> Selection;
+	Selection.Reserve(SelectedIndices.Num());
+	for (int32 i : SelectedIndices)
+	{
+		check(i < ActorInstances.Num());
+		if (ActorInstances[i].IsValid())
+		{
+			Selection.Add(ActorInstances[i].Get());
+		}
+	}
+
+	return Selection;
+}
+
+int32 FFoliageActor::GetInstanceCount() const
+{
+	return ActorInstances.Num();
+}
+
+void FFoliageActor::PreAddInstances(AInstancedFoliageActor* IFA, const UFoliageType* FoliageType, int32 Count)
+{
+	if (!IsInitialized())
+	{
+		Initialize(IFA, FoliageType);
+		check(IsInitialized());
+	}
+}
+
+void FFoliageActor::AddInstance(AInstancedFoliageActor* IFA, const FFoliageInstance& NewInstance)
+{
+	ActorInstances.Add(Spawn(IFA, NewInstance));
+}
+
+void FFoliageActor::RemoveInstance(int32 InstanceIndex)
+{
+	AActor* Actor = ActorInstances[InstanceIndex].Get(false);
+	ActorInstances.RemoveAtSwap(InstanceIndex);
+
+	if (Actor && Actor->GetWorld())
+	{
+		Actor->GetWorld()->DestroyActor(Actor, true);
+	}
+}
+
+void FFoliageActor::SetInstanceWorldTransform(int32 InstanceIndex, const FTransform& Transform, bool bTeleport)
+{
+	if (AActor* Actor = ActorInstances[InstanceIndex].Get(false))
+	{
+		Actor->SetActorTransform(Transform);
+	}
+}
+
+FTransform FFoliageActor::GetInstanceWorldTransform(int32 InstanceIndex) const
+{
+	if (AActor* Actor = ActorInstances[InstanceIndex].Get(false))
+	{
+		return Actor->GetTransform();
+	}
+
+	return FTransform::Identity;
+}
+
+bool FFoliageActor::IsOwnedComponent(const UPrimitiveComponent* Component) const
+{
+	const AActor* Owner = Component->GetOwner();
+
+	return ActorInstances.Contains(Owner);
+}
+
+int32 FFoliageActor::FindIndex(const AActor* InActor) const
+{
+	return ActorInstances.IndexOfByKey(InActor);
+}
+
+int32 FFoliageActor::FindIndex(const UPrimitiveComponent* Component) const
+{
+	return FindIndex(Component->GetOwner());
+}
+
+void FFoliageActor::Refresh(AInstancedFoliageActor* IFA, const TArray<FFoliageInstance>& Instances, bool Async, bool Force)
+{
+	for (int32 i = 0; i < Instances.Num(); ++i)
+	{
+		if (ActorInstances[i] == nullptr || ActorInstances[i]->IsPendingKill())
+		{
+			ActorInstances[i] = Spawn(IFA, Instances[i]);
+		}
+	}
+}
+
+void FFoliageActor::OnHiddenEditorViewMaskChanged(uint64 InHiddenEditorViews)
+{
+	for (const TWeakObjectPtr<AActor>& ActorPtr : ActorInstances)
+	{
+		if (AActor* Actor = ActorPtr.Get(false))
+		{
+			if (Actor->HiddenEditorViews != InHiddenEditorViews)
+			{
+				Actor->HiddenEditorViews = InHiddenEditorViews;
+				Actor->MarkComponentsRenderStateDirty();
+			}
+		}
+	}
+}
+
+void FFoliageActor::PostEditUndo(AInstancedFoliageActor* IFA, UFoliageType* FoliageType, const TArray<FFoliageInstance>& Instances, const TSet<int32>& SelectedIndices)
+{
+	for (int32 i = 0; i < Instances.Num(); ++i)
+	{
+		SetInstanceWorldTransform(i, Instances[i].GetInstanceWorldTransform(), true);
+	}
+}
+
+void FFoliageActor::NotifyFoliageTypeChanged(AInstancedFoliageActor* IFA, UFoliageType* FoliageType, const TArray<FFoliageInstance>& Instances, const TSet<int32>& SelectedIndices, bool bSourceChanged)
+{
+	if (bSourceChanged)
+	{
+		Reapply(IFA, FoliageType, Instances);
+		ApplySelection(true, SelectedIndices);
+	}
+}
+
+void FFoliageActor::Reapply(AInstancedFoliageActor* IFA, const UFoliageType* FoliageType, const TArray<FFoliageInstance>& Instances)
+{
+	IFA->Modify();
+	DestroyActors(false);
+	ActorInstances.Empty();
+
+	if (IsInitialized())
+	{
+		Uninitialize();
+	}
+	Initialize(IFA, FoliageType);
+
+	for (int32 i = 0; i < Instances.Num(); ++i)
+	{
+		ActorInstances.Add(Spawn(IFA, Instances[i]));
+	}
+}
+
+void FFoliageActor::SelectInstances(bool bSelect, int32 InstanceIndex, int32 Count)
+{
+	TArray<AActor*> Selection;
+	Selection.Reserve(Count);
+	check(ActorInstances.Num() >= InstanceIndex + Count);
+	for (int32 i = InstanceIndex; i < (InstanceIndex + Count); ++i)
+	{
+		if (ActorInstances[i].IsValid())
+		{
+			Selection.Add(ActorInstances[i].Get());
+		}
+	}
+	AInstancedFoliageActor::SelectionChanged.Broadcast(bSelect, Selection);
+}
+
+void FFoliageActor::ApplySelection(bool bApply, const TSet<int32>& SelectedIndices)
+{
+	if (bApply && SelectedIndices.Num() > 0)
+	{
+		AInstancedFoliageActor::SelectionChanged.Broadcast(true, GetActorsFromSelectedIndices(SelectedIndices));
+	}
+}
+
+void FFoliageActor::ClearSelection(const TSet<int32>& SelectedIndices)
+{
+	AInstancedFoliageActor::SelectionChanged.Broadcast(false, GetActorsFromSelectedIndices(SelectedIndices));
+}
+
+#endif // WITH_EDITOR
\ No newline at end of file
diff --git a/Engine/Source/Runtime/Foliage/Private/FoliageActor.h b/Engine/Source/Runtime/Foliage/Private/FoliageActor.h
new file mode 100644
--- /dev/null
+++ b/Engine/Source/Runtime/Foliage/Private/FoliageActor.h
@@ -0,0 +1,47 @@
+// Copyright 1998-2019 Epic Games, Inc. All Rights Reserved.
+#pragma once
+
+#include "InstancedFoliage.h"
+
+struct FFoliageActor : public FFoliageImpl
+{
+	TArray<TWeakObjectPtr<AActor>> ActorInstances;
+	UClass* ActorClass;
+	bool bInitialized;
+
+	FFoliageActor()
+		: ActorClass(nullptr)
+		, bInitialized(false)
+	{
+	}
+
+	virtual void AddReferencedObjects(UObject* InThis, FReferenceCollector& Collector) override;
+	virtual void Serialize(FArchive& Ar) override;
+	void DestroyActors(bool bOnLoad);
+
+#if WITH_EDITOR
+	virtual bool IsInitialized() const override;
+	virtual void Initialize(AInstancedFoliageActor* IFA, const UFoliageType* FoliageType) override;
+	virtual void Uninitialize() override;
+	virtual int32 GetInstanceCount() const override;
+	virtual void PreAddInstances(AInstancedFoliageActor* IFA, const UFoliageType* FoliageType, int32 Count) override;
+	virtual void AddInstance(AInstancedFoliageActor* IFA, const FFoliageInstance& NewInstance) override;
+	virtual void RemoveInstance(int32 InstanceIndex) override;
+	virtual void SetInstanceWorldTransform(int32 InstanceIndex, const FTransform& Transform, bool bTeleport) override;
+	virtual FTransform GetInstanceWorldTransform(int32 InstanceIndex) const override;
+	virtual bool IsOwnedComponent(const UPrimitiveComponent* Component) const override;
+	int32 FindIndex(const AActor* InActor) const;
+	virtual int32 FindIndex(const UPrimitiveComponent* HitComponent) const override;
+
+	virtual void SelectInstances(bool bSelect, int32 InstanceIndex, int32 Count) override;
+	virtual void ApplySelection(bool bApply, const TSet<int32>& SelectedIndices) override;
+	virtual void ClearSelection(const TSet<int32>& SelectedIndices) override;
+	virtual void Refresh(AInstancedFoliageActor* IFA, const TArray<FFoliageInstance>& Instances, bool Async, bool Force) override;
+	virtual void OnHiddenEditorViewMaskChanged(uint64 InHiddenEditorViews) override;
+	virtual void PostEditUndo(AInstancedFoliageActor* IFA, UFoliageType* FoliageType, const TArray<FFoliageInstance>& Instances, const TSet<int32>& SelectedIndices) override;
+	virtual void NotifyFoliageTypeChanged(AInstancedFoliageActor* IFA, UFoliageType* FoliageType, const TArray<FFoliageInstance>& Instances, const TSet<int32>& SelectedIndices, bool bSourceChanged) override;
+	void Reapply(AInstancedFoliageActor* IFA, const UFoliageType* FoliageType, const TArray<FFoliageInstance>& Instances);
+	AActor* Spawn(AInstancedFoliageActor* IFA, const FFoliageInstance& Instance);
+	TArray<AActor*> GetActorsFromSelectedIndices(const TSet<int32>& SelectedIndices) const;
+#endif
+};
\ No newline at end of file
diff --git a/Engine/Source/Runtime/Foliage/Private/FoliageInstanceBase.cpp b/Engine/Source/Runtime/Foliage/Private/FoliageInstanceBase.cpp
--- a/Engine/Source/Runtime/Foliage/Private/FoliageInstanceBase.cpp
+++ b/Engine/Source/Runtime/Foliage/Private/FoliageInstanceBase.cpp
@@ -228,9 +228,9 @@ void FFoliageInstanceBaseCache::CompactInstanceBaseCache(AInstancedFoliageActor*
 	if (World != nullptr)
 	{
 		TSet<FFoliageInstanceBaseId> BasesInUse;
-		for (auto& FoliageMeshPair : IFA->FoliageMeshes)
+		for (auto& FoliagePair : IFA->FoliageInfos)
 		{
-			for (const auto& Pair : FoliageMeshPair.Value->ComponentHash)
+			for (const auto& Pair : FoliagePair.Value->ComponentHash)
 			{
 				if (Pair.Key != FFoliageInstanceBaseCache::InvalidBaseId)
 				{
@@ -309,7 +309,7 @@ void FFoliageInstanceBaseCache::CompactInstanceBaseCache(AInstancedFoliageActor*
 
 		if (InvalidBaseIds.Num())
 		{
-			for (auto& Pair : IFA->FoliageMeshes)
+			for (auto& Pair : IFA->FoliageInfos)
 			{
 				auto& MeshInfo = Pair.Value;
 				MeshInfo->ComponentHash.Empty();
diff --git a/Engine/Source/Runtime/Foliage/Private/FoliageStatistics.cpp b/Engine/Source/Runtime/Foliage/Private/FoliageStatistics.cpp
--- a/Engine/Source/Runtime/Foliage/Private/FoliageStatistics.cpp
+++ b/Engine/Source/Runtime/Foliage/Private/FoliageStatistics.cpp
@@ -3,6 +3,7 @@
 #include "FoliageStatistics.h"
 #include "EngineGlobals.h"
 #include "Engine/Engine.h"
+#include "Engine/StaticMesh.h"
 #include "EngineUtils.h"
 #include "InstancedFoliageActor.h"
 
@@ -19,6 +20,7 @@ int32 UFoliageStatistics::FoliageOverlappingSphereCount(UObject* WorldContextObj
 	int32 Count = 0;
 	if (UWorld* World = GEngine->GetWorldFromContextObject(WorldContextObject, EGetWorldErrorMode::LogAndReturnNull))
 	{
+		const UObject* Source = Cast<UObject>(StaticMesh);
 		const FSphere Sphere(CenterPosition, Radius);
 
 		for (TActorIterator<AInstancedFoliageActor> It(World); It; ++It)
@@ -27,7 +29,7 @@ int32 UFoliageStatistics::FoliageOverlappingSphereCount(UObject* WorldContextObj
 			if (!IFA->IsPendingKill())
 			{
 				TArray<const UFoliageType*> FoliageTypes;
-				IFA->GetAllFoliageTypesForMesh(StaticMesh, FoliageTypes);
+				IFA->GetAllFoliageTypesForSource(Source, FoliageTypes);
 
 				for (const auto Type : FoliageTypes)
 				{
@@ -43,6 +45,7 @@ int32 UFoliageStatistics::FoliageOverlappingSphereCount(UObject* WorldContextObj
 int32 UFoliageStatistics::FoliageOverlappingBoxCount(UObject* WorldContextObject, const UStaticMesh* StaticMesh, FBox Box)
 {
 	int32 Count = 0;
+	const UObject* Source = Cast<UObject>(StaticMesh);
 	if (UWorld* World = GEngine->GetWorldFromContextObject(WorldContextObject, EGetWorldErrorMode::LogAndReturnNull))
 	{
 		for (TActorIterator<AInstancedFoliageActor> It(World); It; ++It)
@@ -51,7 +54,7 @@ int32 UFoliageStatistics::FoliageOverlappingBoxCount(UObject* WorldContextObject
 			if (!IFA->IsPendingKill())
 			{
 				TArray<const UFoliageType*> FoliageTypes;
-				IFA->GetAllFoliageTypesForMesh(StaticMesh, FoliageTypes);
+				IFA->GetAllFoliageTypesForSource(Source, FoliageTypes);
 
 				for (const auto Type : FoliageTypes)
 				{
@@ -68,13 +71,14 @@ void UFoliageStatistics::FoliageOverlappingBoxTransforms(UObject* WorldContextOb
 {
 	if (UWorld* World = GEngine->GetWorldFromContextObject(WorldContextObject, EGetWorldErrorMode::LogAndReturnNull))
 	{
+		const UObject* Source = Cast<UObject>(StaticMesh);
 		for (TActorIterator<AInstancedFoliageActor> It(World); It; ++It)
 		{
 			AInstancedFoliageActor* IFA = *It;
 			if (!IFA->IsPendingKill())
 			{
 				TArray<const UFoliageType*> FoliageTypes;
-				IFA->GetAllFoliageTypesForMesh(StaticMesh, FoliageTypes);
+				IFA->GetAllFoliageTypesForSource(Source, FoliageTypes);
 
 				for (const auto Type : FoliageTypes)
 				{
diff --git a/Engine/Source/Runtime/Foliage/Private/FoliageTypeObject.cpp b/Engine/Source/Runtime/Foliage/Private/FoliageTypeObject.cpp
--- a/Engine/Source/Runtime/Foliage/Private/FoliageTypeObject.cpp
+++ b/Engine/Source/Runtime/Foliage/Private/FoliageTypeObject.cpp
@@ -9,7 +9,7 @@ void FFoliageTypeObject::RefreshInstance()
 	TypeInstance = nullptr;
 
 	// If the foliage type is an asset, use that
-	if (auto FoliageTypeAsAsset = Cast<UFoliageType_InstancedStaticMesh>(FoliageTypeObject))
+	if (auto FoliageTypeAsAsset = Cast<UFoliageType>(FoliageTypeObject))
 	{
 		bIsAsset = true;
 		TypeInstance = FoliageTypeAsAsset;
@@ -25,7 +25,7 @@ void FFoliageTypeObject::RefreshInstance()
 	}
 }
 
-const UFoliageType_InstancedStaticMesh* FFoliageTypeObject::GetInstance()
+const UFoliageType* FFoliageTypeObject::GetInstance()
 {
 	if (!TypeInstance)
 	{
@@ -36,7 +36,7 @@ const UFoliageType_InstancedStaticMesh* FFoliageTypeObject::GetInstance()
 	return TypeInstance;
 }
 
-const UFoliageType_InstancedStaticMesh* FFoliageTypeObject::GetInstance() const
+const UFoliageType* FFoliageTypeObject::GetInstance() const
 {
 	return TypeInstance;
 }
@@ -54,14 +54,14 @@ bool FFoliageTypeObject::HasFoliageType() const
 bool FFoliageTypeObject::IsDirty() const
 {
 	//@todo: make sure this works when the object isn't an asset
-	return (!TypeInstance || TypeInstance->ChangeCount != TypeInstance->GetClass()->GetDefaultObject<UFoliageType_InstancedStaticMesh>()->ChangeCount);
+	return (!TypeInstance || TypeInstance->ChangeCount != TypeInstance->GetClass()->GetDefaultObject<UFoliageType>()->ChangeCount);
 }
 
 void FFoliageTypeObject::SetClean()
 {
 	if (TypeInstance)
 	{
-		TypeInstance->ChangeCount = TypeInstance->GetClass()->GetDefaultObject<UFoliageType_InstancedStaticMesh>()->ChangeCount;
+		TypeInstance->ChangeCount = TypeInstance->GetClass()->GetDefaultObject<UFoliageType>()->ChangeCount;
 	}
 }
 
diff --git a/Engine/Source/Runtime/Foliage/Private/InstancedFoliage.cpp b/Engine/Source/Runtime/Foliage/Private/InstancedFoliage.cpp
--- a/Engine/Source/Runtime/Foliage/Private/InstancedFoliage.cpp
+++ b/Engine/Source/Runtime/Foliage/Private/InstancedFoliage.cpp
@@ -21,6 +21,7 @@ InstancedFoliage.cpp: Instanced foliage implementation.
 #include "UObject/UObjectIterator.h"
 #include "FoliageInstancedStaticMeshComponent.h"
 #include "FoliageType_InstancedStaticMesh.h"
+#include "FoliageType_Actor.h"
 #include "InstancedFoliageActor.h"
 #include "Serialization/CustomVersion.h"
 #include "UObject/Package.h"
@@ -42,6 +43,8 @@ InstancedFoliage.cpp: Instanced foliage implementation.
 #include "Engine/StaticMesh.h"
 #include "DrawDebugHelpers.h"
 #include "UObject/FortniteMainBranchObjectVersion.h"
+#include "PreviewScene.h"
+#include "FoliageActor.h"
 
 #define LOCTEXT_NAMESPACE "InstancedFoliage"
 
@@ -55,6 +58,7 @@ DECLARE_CYCLE_STAT(TEXT("FoliageMeshInfo_AddInstance"), STAT_FoliageAddInstance,
 DECLARE_CYCLE_STAT(TEXT("FoliageMeshInfo_RemoveInstance"), STAT_FoliageRemoveInstance, STATGROUP_Foliage);
 DECLARE_CYCLE_STAT(TEXT("FoliageMeshInfo_CreateComponent"), STAT_FoliageCreateComponent, STATGROUP_Foliage);
 
+extern FName FoliageActorTag;
 
 static TAutoConsoleVariable<int32> CVarFoliageDiscardDataOnLoad(
 	TEXT("foliage.DiscardDataOnLoad"),
@@ -91,6 +95,8 @@ struct FFoliageCustomVersion
 		AddedMobility = 10,
 		// Make sure that foliage has FoliageHISMC class
 		FoliageUsingFoliageISMC = 11,
+		// Foliage Actor Support
+		FoliageActorSupport = 12,
 		// -----<new versions can be added above this line>-------------------------------------------------
 		VersionPlusOne,
 		LatestVersion = VersionPlusOne - 1
@@ -107,6 +113,76 @@ const FGuid FFoliageCustomVersion::GUID(0x430C4D19, 0x71544970, 0x87699B69, 0xDF
 // Register the custom version with core
 FCustomVersionRegistration GRegisterFoliageCustomVersion(FFoliageCustomVersion::GUID, FFoliageCustomVersion::LatestVersion, TEXT("FoliageVer"));
 
+///
+/// FFoliageStaticMesh
+///
+struct FFoliageStaticMesh : public FFoliageImpl
+{
+	FFoliageStaticMesh(UHierarchicalInstancedStaticMeshComponent* InComponent)
+		: FFoliageImpl()
+		, Component(InComponent)
+#if WITH_EDITOR
+		, UpdateDepth(0)
+		, bPreviousValue(false)
+		, bInvalidateLightingCache(false)
+#endif
+	{
+
+	}
+
+	UHierarchicalInstancedStaticMeshComponent* Component;
+#if WITH_EDITOR
+	int32 UpdateDepth;
+	bool bPreviousValue;
+
+	bool bInvalidateLightingCache;
+#endif
+
+	virtual void AddReferencedObjects(UObject* InThis, FReferenceCollector& Collector) override;
+	virtual void Serialize(FArchive& Ar) override;
+
+#if WITH_EDITOR
+	virtual bool IsInitialized() const override { return Component != nullptr; }
+	virtual void Initialize(AInstancedFoliageActor* IFA, const UFoliageType* FoliageType) override;
+	virtual void Uninitialize() override;
+	virtual int32 GetInstanceCount() const override;
+	virtual void PreAddInstances(AInstancedFoliageActor* IFA, const UFoliageType* FoliageType, int32 Count) override;
+	virtual void AddInstance(AInstancedFoliageActor* IFA, const FFoliageInstance& NewInstance) override;
+	virtual void RemoveInstance(int32 InstanceIndex) override;
+	virtual void SetInstanceWorldTransform(int32 InstanceIndex, const FTransform& Transform, bool bTeleport) override;
+	virtual FTransform GetInstanceWorldTransform(int32 InstanceIndex) const override;
+	virtual void PostUpdateInstances() override;
+	virtual bool IsOwnedComponent(const UPrimitiveComponent* HitComponent) const override;
+
+	virtual void SelectInstances(bool bSelect, int32 InstanceIndex, int32 Count) override;
+	virtual void ApplySelection(bool bApply, const TSet<int32>& SelectedIndices) override;
+	virtual void ClearSelection(const TSet<int32>& SelectedIndices) override;
+
+	virtual void BeginUpdate() override;
+	virtual void EndUpdate() override;
+	virtual void Refresh(AInstancedFoliageActor* IFA, const TArray<FFoliageInstance>& Instances, bool Async, bool Force) override;
+	virtual void OnHiddenEditorViewMaskChanged(uint64 InHiddenEditorViews) override;
+	virtual void PreEditUndo(AInstancedFoliageActor* IFA, UFoliageType* FoliageType) override;
+	virtual void PostEditUndo(AInstancedFoliageActor* IFA, UFoliageType* FoliageType, const TArray<FFoliageInstance>& Instances, const TSet<int32>& SelectedIndices) override;
+	virtual void NotifyFoliageTypeWillChange(AInstancedFoliageActor* IFA, UFoliageType* FoliageType) override;
+	virtual void NotifyFoliageTypeChanged(AInstancedFoliageActor* IFA, UFoliageType* FoliageType, const TArray<FFoliageInstance>& Instances, const TSet<int32>& SelectedIndices, bool bSourceChanged) override;
+	virtual void EnterEditMode() override;
+	virtual void ExitEditMode() override;
+
+	void HandleComponentMeshBoundsChanged(const FBoxSphereBounds& NewBounds);
+#endif
+
+	virtual int32 GetOverlappingSphereCount(const FSphere& Sphere) const override;
+	virtual int32 GetOverlappingBoxCount(const FBox& Box) const override;
+	virtual void GetOverlappingBoxTransforms(const FBox& Box, TArray<FTransform>& OutTransforms) const override;
+	virtual void GetOverlappingMeshCount(const FSphere& Sphere, TMap<UStaticMesh*, int32>& OutCounts) const override;
+
+	void UpdateComponentSettings(const UFoliageType_InstancedStaticMesh* InSettings);
+	// Recreate the component if the FoliageType's ComponentClass doesn't match the Component's class
+	void CheckComponentClass(AInstancedFoliageActor* InIFA, const UFoliageType_InstancedStaticMesh* InSettings, const TArray<FFoliageInstance>& Instances, const TSet<int32>& SelectedIndices);
+	void ReapplyInstancesToComponent(const TArray<FFoliageInstance>& Instances, const TSet<int32>& SelectedIndices);
+	void CreateNewComponent(AInstancedFoliageActor* InIFA, const UFoliageType* InSettings);
+};
 
 // Legacy (< FFoliageCustomVersion::CrossLevelBase) serializer
 FArchive& operator<<(FArchive& Ar, FFoliageInstance_Deprecated& Instance)
@@ -167,15 +243,17 @@ FArchive& operator<<(FArchive& Ar, FFoliageInstance& Instance)
 static void ConvertDeprecatedFoliageMeshes(
 	AInstancedFoliageActor* IFA,
 	const TMap<UFoliageType*, TUniqueObj<FFoliageMeshInfo_Deprecated>>& FoliageMeshesDeprecated,
-	TMap<UFoliageType*, TUniqueObj<FFoliageMeshInfo>>& FoliageMeshes)
+	TMap<UFoliageType*, TUniqueObj<FFoliageInfo>>& FoliageInfos)
 {
 #if WITH_EDITORONLY_DATA	
 	for (auto Pair : FoliageMeshesDeprecated)
 	{
-		auto& FoliageMesh = FoliageMeshes.Add(Pair.Key);
+		auto& FoliageMesh = FoliageInfos.Add(Pair.Key);
 		const auto& FoliageMeshDeprecated = Pair.Value;
 
-		FoliageMesh->Component = FoliageMeshDeprecated->Component;
+		// Old Foliage mesh is always static mesh (no actors)
+		FoliageMesh->Type = EFoliageImplType::StaticMesh;
+		FoliageMesh->Implementation.Reset(new FFoliageStaticMesh(FoliageMeshDeprecated->Component));
 		FoliageMesh->FoliageTypeUpdateGuid = FoliageMeshDeprecated->FoliageTypeUpdateGuid;
 
 		FoliageMesh->Instances.Reserve(FoliageMeshDeprecated->Instances.Num());
@@ -203,6 +281,32 @@ static void ConvertDeprecatedFoliageMeshes(
 #endif//WITH_EDITORONLY_DATA	
 }
 
+static void ConvertDeprecated2FoliageMeshes(
+	AInstancedFoliageActor* IFA,
+	const TMap<UFoliageType*, TUniqueObj<FFoliageMeshInfo_Deprecated2>>& FoliageMeshesDeprecated,
+	TMap<UFoliageType*, TUniqueObj<FFoliageInfo>>& FoliageInfos)
+{
+#if WITH_EDITORONLY_DATA	
+	for (auto Pair : FoliageMeshesDeprecated)
+	{
+		auto& FoliageMesh = FoliageInfos.Add(Pair.Key);
+		const auto& FoliageMeshDeprecated = Pair.Value;
+
+		// Old Foliage mesh is always static mesh (no actors)
+		FoliageMesh->Type = EFoliageImplType::StaticMesh;
+		FoliageMesh->Implementation.Reset(new FFoliageStaticMesh(FoliageMeshDeprecated->Component));
+		FoliageMesh->FoliageTypeUpdateGuid = FoliageMeshDeprecated->FoliageTypeUpdateGuid;
+
+		FoliageMesh->Instances.Reserve(FoliageMeshDeprecated->Instances.Num());
+
+		for (const FFoliageInstance& Instance : FoliageMeshDeprecated->Instances)
+		{
+			FoliageMesh->Instances.Add(Instance);
+		}
+	}
+#endif//WITH_EDITORONLY_DATA	
+}
+
 /**
 *	FFoliageInstanceCluster_Deprecated
 */
@@ -262,33 +366,55 @@ FArchive& operator<<(FArchive& Ar, FFoliageMeshInfo_Deprecated& MeshInfo)
 	return Ar;
 }
 
-FArchive& operator<<(FArchive& Ar, FFoliageMeshInfo& MeshInfo)
+FFoliageMeshInfo_Deprecated2::FFoliageMeshInfo_Deprecated2()
+	: Component(nullptr)
+{ }
+
+FArchive& operator<<(FArchive& Ar, FFoliageMeshInfo_Deprecated2& MeshInfo)
 {
 	Ar << MeshInfo.Component;
 
+#if WITH_EDITORONLY_DATA
+	Ar << MeshInfo.Instances;
+	Ar << MeshInfo.FoliageTypeUpdateGuid;
+#endif
+
+	return Ar;
+}
+
+FArchive& operator<<(FArchive& Ar, FFoliageInfo& Info)
+{
+	Ar << Info.Type;
+	if (Ar.IsLoading() || (Ar.IsTransacting() && !Info.Implementation.IsValid()))
+	{
+		Info.CreateImplementation(Info.Type);
+	}
+		
+	Info.Implementation->Serialize(Ar);
+
 #if WITH_EDITORONLY_DATA
 	if (!Ar.ArIsFilterEditorOnly && !(Ar.GetPortFlags() & PPF_DuplicateForPIE))
 	{
 		if (Ar.IsTransacting())
 		{
-			MeshInfo.Instances.BulkSerialize(Ar);
+			Info.Instances.BulkSerialize(Ar);
 		}
 		else
 		{
-			Ar << MeshInfo.Instances;
+			Ar << Info.Instances;
 		}
 	}
 
 	if (!Ar.ArIsFilterEditorOnly)
 	{
-		Ar << MeshInfo.FoliageTypeUpdateGuid;
+		Ar << Info.FoliageTypeUpdateGuid;
 	}
 
 	// Serialize the transient data for undo.
 	if (Ar.IsTransacting())
 	{
-		Ar << MeshInfo.ComponentHash;
-		Ar << MeshInfo.SelectedIndices;
+		Ar << Info.ComponentHash;
+		Ar << Info.SelectedIndices;
 	}
 #endif
 
@@ -544,6 +670,110 @@ UFoliageType_InstancedStaticMesh::UFoliageType_InstancedStaticMesh(const FObject
 	CustomNavigableGeometry = EHasCustomNavigableGeometry::Yes;
 }
 
+UObject* UFoliageType_InstancedStaticMesh::GetSource() const
+{
+	return Cast<UObject>(GetStaticMesh());
+}
+
+#if WITH_EDITOR
+void UFoliageType_InstancedStaticMesh::SetSource(UObject* InSource)
+{
+	UStaticMesh* InMesh = Cast<UStaticMesh>(InSource);
+	check(InSource == nullptr || InMesh != nullptr);
+	SetStaticMesh(InMesh);
+}
+
+void UFoliageType_InstancedStaticMesh::UpdateBounds()
+{
+	if (Mesh == nullptr)
+	{
+		return;
+	}
+
+	MeshBounds = Mesh->GetBounds();
+
+	// Make bottom only bound
+	FBox LowBound = MeshBounds.GetBox();
+	LowBound.Max.Z = LowBound.Min.Z + (LowBound.Max.Z - LowBound.Min.Z) * 0.1f;
+
+	float MinX = FLT_MAX, MaxX = FLT_MIN, MinY = FLT_MAX, MaxY = FLT_MIN;
+	LowBoundOriginRadius = FVector::ZeroVector;
+
+	if (Mesh->RenderData)
+	{
+		FPositionVertexBuffer& PositionVertexBuffer = Mesh->RenderData->LODResources[0].VertexBuffers.PositionVertexBuffer;
+		for (uint32 Index = 0; Index < PositionVertexBuffer.GetNumVertices(); ++Index)
+		{
+			const FVector& Pos = PositionVertexBuffer.VertexPosition(Index);
+			if (Pos.Z < LowBound.Max.Z)
+			{
+				MinX = FMath::Min(MinX, Pos.X);
+				MinY = FMath::Min(MinY, Pos.Y);
+				MaxX = FMath::Max(MaxX, Pos.X);
+				MaxY = FMath::Max(MaxY, Pos.Y);
+			}
+		}
+	}
+
+	LowBoundOriginRadius = FVector((MinX + MaxX), (MinY + MaxY), FMath::Sqrt(FMath::Square(MaxX - MinX) + FMath::Square(MaxY - MinY))) * 0.5f;
+}
+#endif
+
+UFoliageType_Actor::UFoliageType_Actor(const FObjectInitializer& ObjectInitializer)
+	: Super(ObjectInitializer)
+{
+	Density = 10;
+	Radius = 500;
+}
+
+#if WITH_EDITOR
+void UFoliageType_Actor::UpdateBounds()
+{
+	if (ActorClass == nullptr)
+	{
+		return;
+	}
+
+	FPreviewScene PreviewScene;
+	FActorSpawnParameters SpawnInfo;
+	SpawnInfo.SpawnCollisionHandlingOverride = ESpawnActorCollisionHandlingMethod::AlwaysSpawn;
+	SpawnInfo.bNoFail = true;
+	SpawnInfo.ObjectFlags = RF_Transient;
+	AActor* PreviewActor = PreviewScene.GetWorld()->SpawnActor<AActor>(ActorClass, SpawnInfo);
+	if (PreviewActor == nullptr)
+	{
+		return;
+	}
+	
+	PreviewActor->SetActorEnableCollision(false);
+	MeshBounds = FBoxSphereBounds(ForceInitToZero);
+
+	// Put this in method...
+	if (PreviewActor != nullptr && PreviewActor->GetRootComponent())
+	{
+		TArray<USceneComponent*> PreviewComponents;
+		PreviewActor->GetRootComponent()->GetChildrenComponents(true, PreviewComponents);
+		PreviewComponents.Add(PreviewActor->GetRootComponent());
+
+		for (USceneComponent* PreviewComponent : PreviewComponents)
+		{
+			MeshBounds = MeshBounds + PreviewComponent->Bounds;
+		}
+	}
+
+	FBox LowBound = MeshBounds.GetBox();
+	LowBound.Max.Z = LowBound.Min.Z + (LowBound.Max.Z - LowBound.Min.Z) * 0.1f;
+
+	float MinX = LowBound.Min.X, MaxX = LowBound.Max.X, MinY = LowBound.Min.Y, MaxY = LowBound.Max.Y;
+	LowBoundOriginRadius = FVector::ZeroVector;
+
+	// TODO: Get more precise lower bound from multiple possible meshes in Actor
+
+	LowBoundOriginRadius = FVector((MinX + MaxX), (MinY + MaxY), FMath::Sqrt(FMath::Square(MaxX - MinX) + FMath::Square(MaxY - MinY))) * 0.5f;
+
+	PreviewActor->Destroy();
+}
+#endif
 
 float UFoliageType::GetMaxRadius() const
 {
@@ -598,18 +828,20 @@ void UFoliageType::PostEditChangeProperty(struct FPropertyChangedEvent& Property
 
 	UpdateGuid = FGuid::NewGuid();
 
-	//@todo: move this into FoliageType_InstancedStaticMesh
-	// Check to see if the mesh is what changed
-	const bool bMeshChanged = PropertyChangedEvent.Property && PropertyChangedEvent.Property->GetFName() == GET_MEMBER_NAME_CHECKED(UFoliageType_InstancedStaticMesh, Mesh);
+	const bool bSourceChanged = IsSourcePropertyChange(PropertyChangedEvent.Property);
+	if (bSourceChanged)
+	{
+		UpdateBounds();
+	}
 
 	// Notify any currently-loaded InstancedFoliageActors
-	if (IsFoliageReallocationRequiredForPropertyChange(PropertyChangedEvent))
+	if (IsFoliageReallocationRequiredForPropertyChange(PropertyChangedEvent.Property))
 	{
 		for (TObjectIterator<AInstancedFoliageActor> It(RF_ClassDefaultObject, /** bIncludeDerivedClasses */ true, /** InternalExcludeFalgs */ EInternalObjectFlags::PendingKill); It; ++It)
 		{
 			if (It->GetWorld() != nullptr)
 			{
-				It->NotifyFoliageTypeChanged(this, bMeshChanged);
+				It->NotifyFoliageTypeChanged(this, bSourceChanged);
 			}
 		}
 	}
@@ -619,11 +851,11 @@ void UFoliageType::PreEditChange(UProperty* PropertyAboutToChange)
 {
 	Super::PreEditChange(PropertyAboutToChange);
 
-	if (PropertyAboutToChange && PropertyAboutToChange->GetFName() == GET_MEMBER_NAME_CHECKED(UFoliageType_InstancedStaticMesh, Mesh))
+	if (IsSourcePropertyChange(PropertyAboutToChange))
 	{
 		for (TObjectIterator<AInstancedFoliageActor> It(RF_ClassDefaultObject, /** bIncludeDerivedClasses */ true, /** InternalExcludeFalgs */ EInternalObjectFlags::PendingKill); It; ++It)
 		{
-			It->NotifyFoliageTypeWillChange(this, true);
+			It->NotifyFoliageTypeWillChange(this);
 		}
 	}
 }
@@ -632,16 +864,10 @@ void UFoliageType::OnHiddenEditorViewMaskChanged(UWorld* InWorld)
 {
 	for (TActorIterator<AInstancedFoliageActor> It(InWorld); It; ++It)
 	{
-		FFoliageMeshInfo* MeshInfo = It->FindMesh(this);
-		if (MeshInfo && MeshInfo->Component)
+		FFoliageInfo* Info = It->FindInfo(this);
+		if (Info != nullptr)
 		{
-			UFoliageInstancedStaticMeshComponent* FoliageComponent = Cast<UFoliageInstancedStaticMeshComponent>(MeshInfo->Component);
-
-			if (FoliageComponent && FoliageComponent->FoliageHiddenEditorViews != HiddenEditorViews)
-			{
-				FoliageComponent->FoliageHiddenEditorViews = HiddenEditorViews;
-				FoliageComponent->MarkRenderStateDirty();
-			}
+			Info->OnHiddenEditorViewMaskChanged(HiddenEditorViews);
 		}
 	}
 }
@@ -658,9 +884,9 @@ FName UFoliageType::GetDisplayFName() const
 	{
 		DisplayFName = FoliageTypeBP->GetFName();
 	}
-	else if (UStaticMesh* StaticMesh = GetStaticMesh())
+	else if (UObject* Source = GetSource())
 	{
-		DisplayFName = StaticMesh->GetFName();
+		DisplayFName = Source->GetFName();
 	}
 
 	return DisplayFName;
@@ -669,122 +895,354 @@ FName UFoliageType::GetDisplayFName() const
 #endif
 
 //
-// FFoliageMeshInfo
+// FFoliageStaticMesh
 //
-FFoliageMeshInfo::FFoliageMeshInfo()
-	: Component(nullptr)
-#if WITH_EDITOR
-	, InstanceHash(GIsEditor ? new FFoliageInstanceHash() : nullptr)
-#endif
-{ }
+void FFoliageStaticMesh::AddReferencedObjects(UObject* InThis, FReferenceCollector& Collector)
+{
+	if (Component != nullptr)
+	{
+		Collector.AddReferencedObject(Component, InThis);
+	}
+}
+
+void FFoliageStaticMesh::Serialize(FArchive& Ar)
+{
+	Ar << Component;
+}
 
-FFoliageMeshInfo::~FFoliageMeshInfo()
-{ }
 
+int32 FFoliageStaticMesh::GetOverlappingSphereCount(const FSphere& Sphere) const
+{
+	if (Component && Component->IsTreeFullyBuilt())
+	{
+		return Component->GetOverlappingSphereCount(Sphere);
+	}
+	return 0;
+}
+
+int32 FFoliageStaticMesh::GetOverlappingBoxCount(const FBox& Box) const
+{
+	if (Component && Component->IsTreeFullyBuilt())
+	{
+		return Component->GetOverlappingBoxCount(Box);
+	}
+	return 0;
+}
+void FFoliageStaticMesh::GetOverlappingBoxTransforms(const FBox& Box, TArray<FTransform>& OutTransforms) const
+{
+	if (Component && Component->IsTreeFullyBuilt())
+	{
+		Component->GetOverlappingBoxTransforms(Box, OutTransforms);
+	}
+}
+void FFoliageStaticMesh::GetOverlappingMeshCount(const FSphere& Sphere, TMap<UStaticMesh*, int32>& OutCounts) const
+{
+	int32 Count = GetOverlappingSphereCount(Sphere);
+	if (Count > 0)
+	{
+		UStaticMesh* const Mesh = Component->GetStaticMesh();
+		int32& StoredCount = OutCounts.FindOrAdd(Mesh);
+		StoredCount += Count;
+	}
+}
 
 #if WITH_EDITOR
+void FFoliageStaticMesh::Initialize(AInstancedFoliageActor* IFA, const UFoliageType* FoliageType)
+{
+	CreateNewComponent(IFA, FoliageType);
+}
 
-void FFoliageMeshInfo::CheckValid()
+void FFoliageStaticMesh::Uninitialize()
 {
-#if DO_FOLIAGE_CHECK
-	int32 ClusterTotal = 0;
-	int32 ComponentTotal = 0;
+	if (Component != nullptr)
+	{
+		if (Component->GetStaticMesh() != nullptr)
+		{
+			Component->GetStaticMesh()->GetOnExtendedBoundsChanged().RemoveAll(this);
+		}
+
+		Component->ClearInstances();
+		Component->SetFlags(RF_Transactional);
+		Component->Modify();
+		Component->DestroyComponent();
+		Component = nullptr;
+	}
+}
 
-	for (FFoliageInstanceCluster& Cluster : InstanceClusters)
+int32 FFoliageStaticMesh::GetInstanceCount() const
+{
+	if (Component != nullptr)
 	{
-		check(Cluster.ClusterComponent != nullptr);
-		ClusterTotal += Cluster.InstanceIndices.Num();
-		ComponentTotal += Cluster.ClusterComponent->PerInstanceSMData.Num();
+		return Component->GetInstanceCount();
 	}
 
-	check(ClusterTotal == ComponentTotal);
+	return 0;
+}
 
-	int32 FreeTotal = 0;
-	int32 InstanceTotal = 0;
-	for (int32 InstanceIdx = 0; InstanceIdx < Instances.Num(); InstanceIdx++)
+void FFoliageStaticMesh::PreAddInstances(AInstancedFoliageActor* IFA, const UFoliageType* FoliageType, int32 Count)
+{
+	if (!IsInitialized())
+	{
+		Initialize(IFA, FoliageType);
+		check(IsInitialized());
+	}
+	else
 	{
-		if (Instances[InstanceIdx].ClusterIndex != -1)
+		Component->InitPerInstanceRenderData(false);
+		Component->InvalidateLightingCache();
+	}
+
+	if (Count)
+	{
+		Component->PreAllocateInstancesMemory(Count);
+	}
+}
+
+void FFoliageStaticMesh::AddInstance(AInstancedFoliageActor* IFA, const FFoliageInstance& NewInstance)
+{
+	check(Component);
+	Component->AddInstanceWorldSpace(NewInstance.GetInstanceWorldTransform());
+}
+
+void FFoliageStaticMesh::RemoveInstance(int32 InstanceIndex)
+{
+	check(Component);
+	Component->RemoveInstance(InstanceIndex);
+
+	if (UpdateDepth > 0)
+	{
+		bInvalidateLightingCache = true;
+	}
+	else
+	{
+		Component->InvalidateLightingCache();
+	}
+}
+
+void FFoliageStaticMesh::SetInstanceWorldTransform(int32 InstanceIndex, const FTransform& Transform, bool bTeleport)
+{
+	check(Component);
+	Component->UpdateInstanceTransform(InstanceIndex, Transform, true, true, bTeleport);
+}
+
+FTransform FFoliageStaticMesh::GetInstanceWorldTransform(int32 InstanceIndex) const
+{
+	return FTransform(Component->PerInstanceSMData[InstanceIndex].Transform) * Component->GetComponentToWorld();
+}
+
+void FFoliageStaticMesh::PostUpdateInstances()
+{
+	check(Component);
+	Component->InvalidateLightingCache();
+	Component->MarkRenderStateDirty();
+}
+
+bool FFoliageStaticMesh::IsOwnedComponent(const UPrimitiveComponent* HitComponent) const
+{
+	return Component == HitComponent;
+}
+
+void FFoliageStaticMesh::SelectInstances(bool bSelect, int32 InstanceIndex, int32 Count)
+{
+	check(Component);
+	Component->SelectInstance(bSelect, InstanceIndex, Count);
+	Component->MarkRenderStateDirty();
+}
+
+void FFoliageStaticMesh::ApplySelection(bool bApply, const TSet<int32>& SelectedIndices)
+{
+	if (Component && (bApply || Component->SelectedInstances.Num() > 0))
+	{
+		Component->ClearInstanceSelection();
+
+		if (bApply)
 		{
-			InstanceTotal++;
+			for (int32 i : SelectedIndices)
+			{
+				Component->SelectInstance(true, i, 1);
+			}
 		}
-		else
+
+		Component->MarkRenderStateDirty();
+	}
+}
+
+void FFoliageStaticMesh::ClearSelection(const TSet<int32>& SelectedIndices)
+{
+	check(Component);
+	Component->ClearInstanceSelection();
+	Component->MarkRenderStateDirty();
+}
+
+void FFoliageStaticMesh::BeginUpdate()
+{
+	if (UpdateDepth == 0)
+	{
+		bPreviousValue = Component->bAutoRebuildTreeOnInstanceChanges;
+		Component->bAutoRebuildTreeOnInstanceChanges = false;
+	}
+	++UpdateDepth;
+}
+
+void FFoliageStaticMesh::EndUpdate()
+{
+	check(UpdateDepth > 0);
+	--UpdateDepth;
+
+	if (UpdateDepth == 0)
+	{
+		Component->bAutoRebuildTreeOnInstanceChanges = bPreviousValue;
+
+		if (bInvalidateLightingCache)
 		{
-			FreeTotal++;
+			Component->InvalidateLightingCache();
+			bInvalidateLightingCache = false;
 		}
 	}
+}
 
-	check(ClusterTotal == InstanceTotal);
-	check(FreeInstanceIndices.Num() == FreeTotal);
+void FFoliageStaticMesh::Refresh(AInstancedFoliageActor* IFA, const TArray<FFoliageInstance>& Instances, bool Async, bool Force)
+{
+	if (Component != nullptr)
+	{
+		Component->BuildTreeIfOutdated(Async, Force);
+	}
+}
 
-	InstanceHash->CheckInstanceCount(InstanceTotal);
+void FFoliageStaticMesh::OnHiddenEditorViewMaskChanged(uint64 InHiddenEditorViews)
+{
+	UFoliageInstancedStaticMeshComponent* FoliageComponent = Cast<UFoliageInstancedStaticMeshComponent>(Component);
 
-	int32 ComponentHashTotal = 0;
-	for (const auto& Pair : ComponentHash)
+	if (FoliageComponent && FoliageComponent->FoliageHiddenEditorViews != InHiddenEditorViews)
 	{
-		ComponentHashTotal += Pair.Value().Num();
+		FoliageComponent->FoliageHiddenEditorViews = InHiddenEditorViews;
+		FoliageComponent->MarkRenderStateDirty();
 	}
-	check(ComponentHashTotal == InstanceTotal);
+}
 
-#if FOLIAGE_CHECK_TRANSFORM
-	// Check transforms match up with editor data
-	int32 MismatchCount = 0;
-	for (int32 ClusterIdx = 0; ClusterIdx<InstanceClusters.Num(); ClusterIdx++)
+void FFoliageStaticMesh::PreEditUndo(AInstancedFoliageActor* IFA, UFoliageType* FoliageType)
+{
+	if (UFoliageType_InstancedStaticMesh* FoliageType_InstancedStaticMesh = Cast<UFoliageType_InstancedStaticMesh>(FoliageType))
+	{
+		if (FoliageType_InstancedStaticMesh->GetStaticMesh() != nullptr)
+		{
+			FoliageType_InstancedStaticMesh->GetStaticMesh()->GetOnExtendedBoundsChanged().RemoveAll(this);
+		}
+	}
+}
+
+void FFoliageStaticMesh::PostEditUndo(AInstancedFoliageActor* IFA, UFoliageType* FoliageType, const TArray<FFoliageInstance>& Instances, const TSet<int32>& SelectedIndices)
+{
+	if (UFoliageType_InstancedStaticMesh* FoliageType_InstancedStaticMesh = Cast<UFoliageType_InstancedStaticMesh>(FoliageType))
 	{
-		TArray<int32> Indices = InstanceClusters(ClusterIdx).InstanceIndices;
-		UInstancedStaticMeshComponent* Comp = InstanceClusters(ClusterIdx).ClusterComponent;
-		for (int32 InstIdx = 0; InstIdx<Indices.Num(); InstIdx++)
+		if (Component != nullptr && FoliageType_InstancedStaticMesh->GetStaticMesh() != nullptr)
 		{
-			int32 InstanceIdx = Indices(InstIdx);
+			FoliageType_InstancedStaticMesh->GetStaticMesh()->GetOnExtendedBoundsChanged().AddRaw(this, &FFoliageStaticMesh::HandleComponentMeshBoundsChanged);
+		}
 
-			FTransform InstanceToWorldEd = Instances(InstanceIdx).GetInstanceTransform();
-			FTransform InstanceToWorldCluster = Comp->PerInstanceSMData(InstIdx).Transform * Comp->GetComponentToWorld();
+		CheckComponentClass(IFA, FoliageType_InstancedStaticMesh, Instances, SelectedIndices);
+		ReapplyInstancesToComponent(Instances, SelectedIndices);
+	}
+}
 
-			if (!InstanceToWorldEd.Equals(InstanceToWorldCluster))
+void FFoliageStaticMesh::NotifyFoliageTypeWillChange(AInstancedFoliageActor* IFA, UFoliageType* FoliageType)
+{
+	if (Component != nullptr)
+	{
+		if (UFoliageType_InstancedStaticMesh* FoliageType_InstancedStaticMesh = Cast<UFoliageType_InstancedStaticMesh>(FoliageType))
+		{
+			if (FoliageType_InstancedStaticMesh->GetStaticMesh() != nullptr)
 			{
-				Comp->PerInstanceSMData(InstIdx).Transform = InstanceToWorldEd.ToMatrixWithScale();
-				MismatchCount++;
+				FoliageType_InstancedStaticMesh->GetStaticMesh()->GetOnExtendedBoundsChanged().RemoveAll(this);
 			}
 		}
 	}
+}
 
-	if (MismatchCount != 0)
+void FFoliageStaticMesh::NotifyFoliageTypeChanged(AInstancedFoliageActor* IFA, UFoliageType* FoliageType, const TArray<FFoliageInstance>& Instances, const TSet<int32>& SelectedIndices, bool bSourceChanged)
+{
+	UFoliageType_InstancedStaticMesh* FoliageType_InstancedStaticMesh = Cast<UFoliageType_InstancedStaticMesh>(FoliageType);
+	check(FoliageType_InstancedStaticMesh);
+	CheckComponentClass(IFA, FoliageType_InstancedStaticMesh, Instances, SelectedIndices);
+	UpdateComponentSettings(FoliageType_InstancedStaticMesh);
+	
+	if (bSourceChanged && Component != nullptr && Component->GetStaticMesh() != nullptr)
+	{
+		// Change bounds delegate bindings
+		if (FoliageType_InstancedStaticMesh->GetStaticMesh() != nullptr)
+		{
+			Component->GetStaticMesh()->GetOnExtendedBoundsChanged().AddRaw(this, &FFoliageStaticMesh::HandleComponentMeshBoundsChanged);
+
+			// Mesh changed, so we must update the occlusion tree
+			Component->BuildTreeIfOutdated(true, false);
+		}
+	}
+}
+
+void FFoliageStaticMesh::EnterEditMode()
+{
+	if (Component == nullptr)
 	{
-		UE_LOG(LogInstancedFoliage, Log, TEXT("%s: transform mismatch: %d"), *InstanceClusters(0).ClusterComponent->StaticMesh->GetName(), MismatchCount);
+		return;
 	}
-#endif
 
-#endif
+	if (Component->GetStaticMesh() != nullptr)
+	{
+		Component->GetStaticMesh()->GetOnExtendedBoundsChanged().AddRaw(this, &FFoliageStaticMesh::HandleComponentMeshBoundsChanged);
+
+		Component->BuildTreeIfOutdated(true, false);
+	}
+
+	Component->bCanEnableDensityScaling = false;
+	Component->UpdateDensityScaling();
 }
 
-void FFoliageMeshInfo::CreateNewComponent(AInstancedFoliageActor* InIFA, const UFoliageType* InSettings)
+void FFoliageStaticMesh::ExitEditMode()
+{
+	if (Component == nullptr)
+	{
+		return;
+	}
+
+	if (Component->GetStaticMesh() != nullptr)
+	{
+		Component->GetStaticMesh()->GetOnExtendedBoundsChanged().RemoveAll(this);
+	}
+
+	Component->bCanEnableDensityScaling = true;
+	Component->UpdateDensityScaling();
+}
+
+void FFoliageStaticMesh::CreateNewComponent(AInstancedFoliageActor* InIFA, const UFoliageType* InSettings)
 {
 	SCOPE_CYCLE_COUNTER(STAT_FoliageCreateComponent);
-	
+
 	check(!Component);
+	const UFoliageType_InstancedStaticMesh* FoliageType_InstancedStaticMesh = Cast<UFoliageType_InstancedStaticMesh>(InSettings);
 
-	UClass* ComponentClass = InSettings->GetComponentClass();
+	UClass* ComponentClass = FoliageType_InstancedStaticMesh->GetComponentClass();
 	if (ComponentClass == nullptr)
 	{
 		ComponentClass = UFoliageInstancedStaticMeshComponent::StaticClass();
 	}
 
 	UFoliageInstancedStaticMeshComponent* FoliageComponent = NewObject<UFoliageInstancedStaticMeshComponent>(InIFA, ComponentClass, NAME_None, RF_Transactional);
+	
+	check(FoliageType_InstancedStaticMesh);
 
 	Component = FoliageComponent;
-	Component->SetStaticMesh(InSettings->GetStaticMesh());
+	Component->SetStaticMesh(FoliageType_InstancedStaticMesh->GetStaticMesh());
 	Component->bSelectable = true;
 	Component->bHasPerInstanceHitProxies = true;
 
 	if (Component->GetStaticMesh() != nullptr)
 	{
-		Component->GetStaticMesh()->GetOnExtendedBoundsChanged().AddRaw(this, &FFoliageMeshInfo::HandleComponentMeshBoundsChanged);
+		Component->GetStaticMesh()->GetOnExtendedBoundsChanged().AddRaw(this, &FFoliageStaticMesh::HandleComponentMeshBoundsChanged);
 	}
 
-#if WITH_EDITOR
 	FoliageComponent->FoliageHiddenEditorViews = InSettings->HiddenEditorViews;
-#endif	
 
-	UpdateComponentSettings(InSettings);
+	UpdateComponentSettings(FoliageType_InstancedStaticMesh);
 
 	Component->SetupAttachment(InIFA->GetRootComponent());
 
@@ -802,7 +1260,7 @@ void FFoliageMeshInfo::CreateNewComponent(AInstancedFoliageActor* InIFA, const U
 	Component->ClearFlags(RF_Transactional);
 }
 
-void FFoliageMeshInfo::HandleComponentMeshBoundsChanged(const FBoxSphereBounds& NewBounds)
+void FFoliageStaticMesh::HandleComponentMeshBoundsChanged(const FBoxSphereBounds& NewBounds)
 {
 	if (Component != nullptr)
 	{
@@ -810,7 +1268,7 @@ void FFoliageMeshInfo::HandleComponentMeshBoundsChanged(const FBoxSphereBounds&
 	}
 }
 
-void FFoliageMeshInfo::CheckComponentClass(AInstancedFoliageActor* InIFA, const UFoliageType* InSettings)
+void FFoliageStaticMesh::CheckComponentClass(AInstancedFoliageActor* InIFA, const UFoliageType_InstancedStaticMesh* InSettings, const TArray<FFoliageInstance>& Instances, const TSet<int32>& SelectedIndices)
 {
 	if (Component)
 	{
@@ -825,37 +1283,29 @@ void FFoliageMeshInfo::CheckComponentClass(AInstancedFoliageActor* InIFA, const
 			InIFA->Modify();
 
 			// prepare to destroy the old component
-			Component->ClearInstances();
-
-			// make sure the destruction gets stored in the undo buffer, 
-			// so the component will get recreated on undo.
-			Component->SetFlags(RF_Transactional);
-			Component->Modify();
-
-			Component->DestroyComponent();
-			Component = nullptr;
+			Uninitialize();
 
 			// create a new component
-			CreateNewComponent(InIFA, InSettings);
+			Initialize(InIFA, InSettings);
 
 			// apply the instances to it
-			ReapplyInstancesToComponent();
+			ReapplyInstancesToComponent(Instances, SelectedIndices);
 		}
 	}
 }
 
-void FFoliageMeshInfo::UpdateComponentSettings(const UFoliageType* InSettings)
+void FFoliageStaticMesh::UpdateComponentSettings(const UFoliageType_InstancedStaticMesh* InSettings)
 {
 	if (Component)
 	{
 		bool bNeedsMarkRenderStateDirty = false;
 		bool bNeedsInvalidateLightingCache = false;
 
-		const UFoliageType* FoliageType = InSettings;
+		const UFoliageType_InstancedStaticMesh* FoliageType = InSettings;
 		if (InSettings->GetClass()->ClassGeneratedBy)
 		{
 			// If we're updating settings for a BP foliage type, use the CDO
-			FoliageType = InSettings->GetClass()->GetDefaultObject<UFoliageType>();
+			FoliageType = InSettings->GetClass()->GetDefaultObject<UFoliageType_InstancedStaticMesh>();
 		}
 
 		if (Component->GetStaticMesh() != FoliageType->GetStaticMesh())
@@ -1024,100 +1474,288 @@ void FFoliageMeshInfo::UpdateComponentSettings(const UFoliageType* InSettings)
 	}
 }
 
-void FFoliageMeshInfo::AddInstances(AInstancedFoliageActor* InIFA, const UFoliageType* InSettings, const TSet<const FFoliageInstance*>& InNewInstances, bool RebuildFoliageTree)
+void FFoliageStaticMesh::ReapplyInstancesToComponent(const TArray<FFoliageInstance>& Instances, const TSet<int32>& SelectedIndices)
 {
-	SCOPE_CYCLE_COUNTER(STAT_FoliageAddInstance);
+	if (Component)
+	{
+		// clear the transactional flag if it was set prior to deleting the actor
+		Component->ClearFlags(RF_Transactional);
 
-	InIFA->Modify();
+		const bool bWasRegistered = Component->IsRegistered();
+		Component->UnregisterComponent();
+		Component->ClearInstances();
+		Component->InitPerInstanceRenderData(false);
 
-	if (Component == nullptr)
+		Component->bAutoRebuildTreeOnInstanceChanges = false;
+
+		for (auto& Instance : Instances)
+		{
+			Component->AddInstanceWorldSpace(Instance.GetInstanceWorldTransform());
+		}
+
+		Component->bAutoRebuildTreeOnInstanceChanges = true;
+		Component->BuildTreeIfOutdated(true, true);
+
+		Component->ClearInstanceSelection();
+
+		if (SelectedIndices.Num())
+		{
+			for (int32 i : SelectedIndices)
+			{
+				Component->SelectInstance(true, i, 1);
+			}
+		}
+
+		if (bWasRegistered)
+		{
+			Component->RegisterComponent();
+		}
+	}
+}
+
+#endif // WITH_EDITOR
+
+//
+// FFoliageInfo
+//
+FFoliageInfo::FFoliageInfo()
+	: Type(EFoliageImplType::StaticMesh)
+#if WITH_EDITOR
+	, InstanceHash(GIsEditor ? new FFoliageInstanceHash() : nullptr)
+#endif
+{ }
+
+FFoliageInfo::~FFoliageInfo()
+{ }
+
+UHierarchicalInstancedStaticMeshComponent* FFoliageInfo::GetComponent() const
+{
+	if (Type == EFoliageImplType::StaticMesh)
 	{
-		CreateNewComponent(InIFA, InSettings);
-		check(Component);
+		FFoliageStaticMesh* FoliageStaticMesh = StaticCast<FFoliageStaticMesh*>(Implementation.Get());
+		return FoliageStaticMesh->Component;
 	}
-	else
+
+	return nullptr;
+}
+
+void FFoliageInfo::AddReferencedObjects(UObject* InThis, FReferenceCollector& Collector)
+{
+	Implementation->AddReferencedObjects(InThis, Collector);
+}
+
+void FFoliageInfo::CreateImplementation(EFoliageImplType InType)
+{
+	check(InType != EFoliageImplType::Unknown);
+	check(!Implementation.IsValid());
+	// Change Impl based on InType param
+	Type = InType;
+	
+	if (Type == EFoliageImplType::StaticMesh)
 	{
-		Component->InitPerInstanceRenderData(false);
-		Component->InvalidateLightingCache();
+		Implementation.Reset(new FFoliageStaticMesh(nullptr));
+	}
+	else if (Type == EFoliageImplType::Actor)
+	{
+		Implementation.Reset(new FFoliageActor());
 	}
+}
 
-	bool PreviousbAutoRebuildTreeOnInstanceChanges = Component->bAutoRebuildTreeOnInstanceChanges;
-	Component->bAutoRebuildTreeOnInstanceChanges = RebuildFoliageTree;
+int32 FFoliageInfo::GetOverlappingSphereCount(const FSphere& Sphere) const
+{
+	return Implementation->GetOverlappingSphereCount(Sphere);
+}
 
-	Instances.Reserve(Instances.Num() + InNewInstances.Num());
-	Component->PreAllocateInstancesMemory(InNewInstances.Num());
+int32 FFoliageInfo::GetOverlappingBoxCount(const FBox& Box) const
+{
+	return Implementation->GetOverlappingBoxCount(Box);
+}
 
-	for (const FFoliageInstance* Instance : InNewInstances)
+void FFoliageInfo::GetOverlappingBoxTransforms(const FBox& Box, TArray<FTransform>& OutTransforms) const
+{
+	Implementation->GetOverlappingBoxTransforms(Box, OutTransforms);
+}
+
+void FFoliageInfo::GetOverlappingMeshCount(const FSphere& Sphere, TMap<UStaticMesh*, int32>& OutCounts) const
+{
+	Implementation->GetOverlappingMeshCount(Sphere, OutCounts);
+}
+
+#if WITH_EDITOR
+void FFoliageInfo::CreateImplementation(const UFoliageType* FoliageType)
+{
+	check(!Implementation.IsValid());
+	// Change Impl based on FoliageType param
+	Type = EFoliageImplType::Unknown;
+	if (FoliageType->IsA<UFoliageType_InstancedStaticMesh>())
+	{
+		Type = EFoliageImplType::StaticMesh;
+		Implementation.Reset(new FFoliageStaticMesh(nullptr));
+	}
+	else if (FoliageType->IsA<UFoliageType_Actor>())
+	{
+		Type = EFoliageImplType::Actor;
+		Implementation.Reset(new FFoliageActor());
+	}
+	check(Type != EFoliageImplType::Unknown);
+}
+
+void FFoliageInfo::Initialize(AInstancedFoliageActor* IFA, const UFoliageType* FoliageType)
+{
+	check(!IsInitialized());
+	check(Implementation.IsValid());
+	Implementation->Initialize(IFA, FoliageType);
+}
+
+void FFoliageInfo::Uninitialize()
+{
+	check(IsInitialized());
+	Implementation->Uninitialize();
+}
+
+bool FFoliageInfo::IsInitialized() const
+{
+	return Implementation.IsValid() && Implementation->IsInitialized();
+}
+
+void FFoliageInfo::NotifyFoliageTypeWillChange(AInstancedFoliageActor* IFA, UFoliageType* FoliageType)
+{
+	Implementation->NotifyFoliageTypeWillChange(IFA, FoliageType);
+}
+
+void FFoliageInfo::NotifyFoliageTypeChanged(AInstancedFoliageActor* IFA, UFoliageType* FoliageType, bool bSourceChanged)
+{
+	Implementation->NotifyFoliageTypeChanged(IFA, FoliageType, Instances, SelectedIndices, bSourceChanged);
+}
+
+void FFoliageInfo::CheckValid()
+{
+#if DO_FOLIAGE_CHECK
+	int32 ClusterTotal = 0;
+	int32 ComponentTotal = 0;
+
+	check(Instances.Num() == Implementation->GetInstanceCount());
+		
+	InstanceHash->CheckInstanceCount(Instances.Num());
+
+	int32 ComponentHashTotal = 0;
+	for (const auto& Pair : ComponentHash)
+	{
+		ComponentHashTotal += Pair.Value.Num();
+	}
+	check(ComponentHashTotal == Instances.Num());
+
+#if FOLIAGE_CHECK_TRANSFORM
+	// Check transforms match up with editor data
+	int32 MismatchCount = 0;
+	for (int32 i = 0; i < Instances.Num(); ++i)
+	{
+		FTransform InstanceToWorldEd = Instances[i].GetInstanceWorldTransform();
+		FTransform InstanceToWorldImpl = Implementation->GetInstanceWorldTransform(i);
+
+		if (!InstanceToWorldEd.Equals(InstanceToWorldImpl))
+		{
+			MismatchCount++;
+		}
+	}
+		
+	if (MismatchCount != 0)
+	{
+		UE_LOG(LogInstancedFoliage, Log, TEXT("transform mismatch: %d"), MismatchCount);
+	}
+#endif
+
+#endif
+}
+
+/*
+*/
+void FFoliageInfo::ClearSelection()
+{
+	if (Instances.Num() > 0)
+	{
+		Implementation->ClearSelection(SelectedIndices);
+		SelectedIndices.Empty();
+	}
+}
+
+void FFoliageInfo::SetRandomSeed(int32 seed)
+{
+	if (Type == EFoliageImplType::StaticMesh)
 	{
-		// Add the instance taking either a free slot or adding a new item.
-		int32 InstanceIndex = Instances.Add(*Instance);
-		FFoliageInstance& AddedInstance = Instances[InstanceIndex];
+		FFoliageStaticMesh* FoliageStaticMesh = StaticCast<FFoliageStaticMesh*>(Implementation.Get());
+		FoliageStaticMesh->Component->InstancingRandomSeed = 1;
+	}
+}
+
+void FFoliageInfo::SetInstanceWorldTransform(int32 InstanceIndex, const FTransform& Transform, bool bTeleport)
+{
+	Implementation->SetInstanceWorldTransform(InstanceIndex, Transform, bTeleport);
+}
+
+void FFoliageInfo::AddInstanceImpl(AInstancedFoliageActor* InIFA, const FFoliageInstance& InNewInstance)
+{
+	// Add the instance taking either a free slot or adding a new item.
+	int32 InstanceIndex = Instances.Add(InNewInstance);
+	FFoliageInstance& AddedInstance = Instances[InstanceIndex];
 
-		AddedInstance.BaseId = InIFA->InstanceBaseCache.AddInstanceBaseId(Instance->BaseComponent);
+	AddedInstance.BaseId = InIFA->InstanceBaseCache.AddInstanceBaseId(InNewInstance.BaseComponent);
 
-		// Add the instance to the hash
-		AddToBaseHash(InstanceIndex);
-		InstanceHash->InsertInstance(AddedInstance.Location, InstanceIndex);
+	// Add the instance to the hash
+	AddToBaseHash(InstanceIndex);
+	InstanceHash->InsertInstance(AddedInstance.Location, InstanceIndex);
 
-		// Calculate transform for the instance
-		FTransform InstanceToWorld = Instance->GetInstanceWorldTransform();
+	// Add the instance to the component
+	Implementation->AddInstance(InIFA, AddedInstance);
+}
 
-		// Add the instance to the component
-		Component->AddInstanceWorldSpace(InstanceToWorld);
+void FFoliageInfo::AddInstances(AInstancedFoliageActor* InIFA, const UFoliageType* InSettings, const TSet<const FFoliageInstance*>& InNewInstances)
+{
+	SCOPE_CYCLE_COUNTER(STAT_FoliageAddInstance);
+
+	InIFA->Modify();
+
+	Implementation->PreAddInstances(InIFA, InSettings, InNewInstances.Num());
+
+	Implementation->BeginUpdate();
+
+	Instances.Reserve(Instances.Num() + InNewInstances.Num());
+			
+	for (const FFoliageInstance* Instance : InNewInstances)
+	{
+		AddInstanceImpl(InIFA, *Instance);
 	}
 
 	CheckValid();
 
-	Component->bAutoRebuildTreeOnInstanceChanges = PreviousbAutoRebuildTreeOnInstanceChanges;
+	Implementation->EndUpdate();
 }
 
-void FFoliageMeshInfo::AddInstance(AInstancedFoliageActor* InIFA, const UFoliageType* InSettings, const FFoliageInstance& InNewInstance, UActorComponent* InBaseComponent, bool RebuildFoliageTree)
+void FFoliageInfo::AddInstance(AInstancedFoliageActor* InIFA, const UFoliageType* InSettings, const FFoliageInstance& InNewInstance, UActorComponent* InBaseComponent)
 {
 	FFoliageInstance Instance = InNewInstance;
 	Instance.BaseId = InIFA->InstanceBaseCache.AddInstanceBaseId(InBaseComponent);
-	AddInstance(InIFA, InSettings, Instance, RebuildFoliageTree);
+	AddInstance(InIFA, InSettings, Instance);
 }
 
-void FFoliageMeshInfo::AddInstance(AInstancedFoliageActor* InIFA, const UFoliageType* InSettings, const FFoliageInstance& InNewInstance, bool RebuildFoliageTree)
+void FFoliageInfo::AddInstance(AInstancedFoliageActor* InIFA, const UFoliageType* InSettings, const FFoliageInstance& InNewInstance)
 {
 	SCOPE_CYCLE_COUNTER(STAT_FoliageAddInstance);
 
 	InIFA->Modify();
 
-	if (Component == nullptr)
-	{
-		CreateNewComponent(InIFA, InSettings);
-		check(Component);
-	}
-	else
-	{
-		Component->InitPerInstanceRenderData(false);
-		Component->InvalidateLightingCache();
-	}
-
-	bool PreviousbAutoRebuildTreeOnInstanceChanges = Component->bAutoRebuildTreeOnInstanceChanges;
-	Component->bAutoRebuildTreeOnInstanceChanges = RebuildFoliageTree;
-
-	// Add the instance taking either a free slot or adding a new item.
-	int32 InstanceIndex = Instances.Add(InNewInstance);
-	FFoliageInstance& AddedInstance = Instances[InstanceIndex];
-
-	// Add the instance to the hash
-	AddToBaseHash(InstanceIndex);
-	InstanceHash->InsertInstance(AddedInstance.Location, InstanceIndex);
+	Implementation->PreAddInstances(InIFA, InSettings, 1);
 
-	// Calculate transform for the instance
-	FTransform InstanceToWorld = InNewInstance.GetInstanceWorldTransform();
+	Implementation->BeginUpdate();
 
-	// Add the instance to the component
-	Component->AddInstanceWorldSpace(InstanceToWorld);
+	AddInstanceImpl(InIFA, InNewInstance);
 
 	CheckValid();
 
-	Component->bAutoRebuildTreeOnInstanceChanges = PreviousbAutoRebuildTreeOnInstanceChanges;
+	Implementation->EndUpdate();
 }
 
-void FFoliageMeshInfo::RemoveInstances(AInstancedFoliageActor* InIFA, const TArray<int32>& InInstancesToRemove, bool RebuildFoliageTree)
+void FFoliageInfo::RemoveInstances(AInstancedFoliageActor* InIFA, const TArray<int32>& InInstancesToRemove, bool RebuildFoliageTree)
 {
 	SCOPE_CYCLE_COUNTER(STAT_FoliageRemoveInstance);
 
@@ -1126,11 +1764,10 @@ void FFoliageMeshInfo::RemoveInstances(AInstancedFoliageActor* InIFA, const TArr
 		return;
 	}
 
-	check(Component);
+	check(IsInitialized());
 	InIFA->Modify();
 
-	bool PreviousbAutoRebuildTreeOnInstanceChanges = Component->bAutoRebuildTreeOnInstanceChanges;
-	Component->bAutoRebuildTreeOnInstanceChanges = false;
+	Implementation->BeginUpdate();
 
 	TSet<int32> InstancesToRemove;
 	InstancesToRemove.Append(InInstancesToRemove);
@@ -1149,7 +1786,7 @@ void FFoliageMeshInfo::RemoveInstances(AInstancedFoliageActor* InIFA, const TArr
 		InstanceHash->RemoveInstance(Instance.Location, InstanceIndex);
 
 		// remove from the component
-		Component->RemoveInstance(InstanceIndex);
+		Implementation->RemoveInstance(InstanceIndex);
 
 		// Remove it from the selection.
 		SelectedIndices.Remove(InstanceIndex);
@@ -1195,18 +1832,17 @@ void FFoliageMeshInfo::RemoveInstances(AInstancedFoliageActor* InIFA, const TArr
 
 	Instances.Shrink();
 		
-	Component->bAutoRebuildTreeOnInstanceChanges = PreviousbAutoRebuildTreeOnInstanceChanges;
-	Component->InvalidateLightingCache();
-
+	Implementation->EndUpdate();
+		
 	if (RebuildFoliageTree)
 	{
-		Component->BuildTreeIfOutdated(true, true);
-	}		
+		Refresh(InIFA, true, true);
+	}
 
 	CheckValid();
 }
 
-void FFoliageMeshInfo::PreMoveInstances(AInstancedFoliageActor* InIFA, const TArray<int32>& InInstancesToMove)
+void FFoliageInfo::PreMoveInstances(AInstancedFoliageActor* InIFA, const TArray<int32>& InInstancesToMove)
 {
 	// Remove instances from the hash
 	for (TArray<int32>::TConstIterator It(InInstancesToMove); It; ++It)
@@ -1218,12 +1854,10 @@ void FFoliageMeshInfo::PreMoveInstances(AInstancedFoliageActor* InIFA, const TAr
 }
 
 
-void FFoliageMeshInfo::PostUpdateInstances(AInstancedFoliageActor* InIFA, const TArray<int32>& InInstancesUpdated, bool bReAddToHash, bool InUpdateSelection)
+void FFoliageInfo::PostUpdateInstances(AInstancedFoliageActor* InIFA, const TArray<int32>& InInstancesUpdated, bool bReAddToHash, bool InUpdateSelection)
 {
 	if (InInstancesUpdated.Num())
 	{
-		check(Component);
-
 		for (TArray<int32>::TConstIterator It(InInstancesUpdated); It; ++It)
 		{
 			int32 InstanceIndex = *It;
@@ -1231,7 +1865,7 @@ void FFoliageMeshInfo::PostUpdateInstances(AInstancedFoliageActor* InIFA, const
 
 			FTransform InstanceToWorld = Instance.GetInstanceWorldTransform();
 
-			Component->UpdateInstanceTransform(InstanceIndex, InstanceToWorld, true);
+			Implementation->SetInstanceWorldTransform(InstanceIndex, InstanceToWorld, true);
 
 			// Re-add instance to the hash if requested
 			if (bReAddToHash)
@@ -1242,42 +1876,35 @@ void FFoliageMeshInfo::PostUpdateInstances(AInstancedFoliageActor* InIFA, const
 			// Reselect the instance to update the render update to include selection as by default it gets removed
 			if (InUpdateSelection)
 			{
-				Component->SelectInstance(true, InstanceIndex);
+				Implementation->SelectInstances(true, InstanceIndex, 1);
 			}
 		}
 
-		Component->InvalidateLightingCache();
-		Component->MarkRenderStateDirty();
+		Implementation->PostUpdateInstances();
 	}
 }
 
-void FFoliageMeshInfo::PostMoveInstances(AInstancedFoliageActor* InIFA, const TArray<int32>& InInstancesMoved)
+void FFoliageInfo::PostMoveInstances(AInstancedFoliageActor* InIFA, const TArray<int32>& InInstancesMoved)
 {
 	PostUpdateInstances(InIFA, InInstancesMoved, true, true);
 }
 
-void FFoliageMeshInfo::DuplicateInstances(AInstancedFoliageActor* InIFA, UFoliageType* InSettings, const TArray<int32>& InInstancesToDuplicate)
+void FFoliageInfo::DuplicateInstances(AInstancedFoliageActor* InIFA, UFoliageType* InSettings, const TArray<int32>& InInstancesToDuplicate)
 {
-	if (Component != nullptr)
-	{
-		Component->bAutoRebuildTreeOnInstanceChanges = false;
-	}
+	Implementation->BeginUpdate();
 
 	for (int32 InstanceIndex : InInstancesToDuplicate)
 	{
 		const FFoliageInstance TempInstance = Instances[InstanceIndex];
-		AddInstance(InIFA, InSettings, TempInstance, false);
+		AddInstance(InIFA, InSettings, TempInstance);
 	}
 
-	if (Component != nullptr)
-	{
-		Component->bAutoRebuildTreeOnInstanceChanges = true;
-		Component->BuildTreeIfOutdated(true, true);
-	}
+	Implementation->EndUpdate();
+	Refresh(InIFA, true, true);
 }
 
 /* Get the number of placed instances */
-int32 FFoliageMeshInfo::GetPlacedInstanceCount() const
+int32 FFoliageInfo::GetPlacedInstanceCount() const
 {
 	int32 PlacedInstanceCount = 0;
 
@@ -1292,13 +1919,13 @@ int32 FFoliageMeshInfo::GetPlacedInstanceCount() const
 	return PlacedInstanceCount;
 }
 
-void FFoliageMeshInfo::AddToBaseHash(int32 InstanceIndex)
+void FFoliageInfo::AddToBaseHash(int32 InstanceIndex)
 {
 	FFoliageInstance& Instance = Instances[InstanceIndex];
 	ComponentHash.FindOrAdd(Instance.BaseId).Add(InstanceIndex);
 }
 
-void FFoliageMeshInfo::RemoveFromBaseHash(int32 InstanceIndex)
+void FFoliageInfo::RemoveFromBaseHash(int32 InstanceIndex)
 {
 	FFoliageInstance& Instance = Instances[InstanceIndex];
 
@@ -1316,17 +1943,12 @@ void FFoliageMeshInfo::RemoveFromBaseHash(int32 InstanceIndex)
 }
 
 // Destroy existing clusters and reassign all instances to new clusters
-void FFoliageMeshInfo::ReallocateClusters(AInstancedFoliageActor* InIFA, UFoliageType* InSettings)
+void FFoliageInfo::ReallocateClusters(AInstancedFoliageActor* InIFA, UFoliageType* InSettings)
 {
-	if (Component != nullptr)
+	if (Implementation->IsInitialized())
 	{
-		Component->ClearInstances();
-		Component->SetFlags(RF_Transactional);
-		Component->Modify();
-		Component->DestroyComponent();
-		Component = nullptr;
+		Implementation->Uninitialize();
 	}
-
 	// Remove everything
 	TArray<FFoliageInstance> OldInstances;
 	Exchange(Instances, OldInstances);
@@ -1342,14 +1964,12 @@ void FFoliageMeshInfo::ReallocateClusters(AInstancedFoliageActor* InIFA, UFoliag
 	{
 		if ((Instance.Flags & FOLIAGE_InstanceDeleted) == 0)
 		{
-			AddInstance(InIFA, InSettings, Instance, false);
+			AddInstance(InIFA, InSettings, Instance);
 		}
 	}
 
-	if (Component != nullptr)
-	{	
-		Component->BuildTreeIfOutdated(true, true);
-	}
+	
+	Refresh(InIFA, true, true);
 }
 
 void FFoliageMeshInfo::ReapplyInstancesToComponent()
@@ -1372,7 +1992,7 @@ void FFoliageMeshInfo::ReapplyInstancesToComponent()
 		}
 
 		Component->bAutoRebuildTreeOnInstanceChanges = true;
-		Component->BuildTreeIfOutdated(false, true);
+		Component->BuildTreeIfOutdated(true, true);
 
 		Component->ClearInstanceSelection();
 
@@ -1403,7 +2023,7 @@ void FFoliageMeshInfo::GetInstancesInsideSphere(const FSphere& Sphere, TArray<in
 	}
 }
 
-void FFoliageMeshInfo::GetInstanceAtLocation(const FVector& Location, int32& OutInstance, bool& bOutSucess)
+void FFoliageInfo::GetInstanceAtLocation(const FVector& Location, int32& OutInstance, bool& bOutSucess)
 {
 	auto TempInstances = InstanceHash->GetInstancesOverlappingBox(FBox::BuildAABB(Location, FVector(KINDA_SMALL_NUMBER)));
 
@@ -1425,7 +2045,7 @@ void FFoliageMeshInfo::GetInstanceAtLocation(const FVector& Location, int32& Out
 }
 
 // Returns whether or not there is are any instances overlapping the sphere specified
-bool FFoliageMeshInfo::CheckForOverlappingSphere(const FSphere& Sphere)
+bool FFoliageInfo::CheckForOverlappingSphere(const FSphere& Sphere)
 {
 	auto TempInstances = InstanceHash->GetInstancesOverlappingBox(FBox::BuildAABB(Sphere.Center, FVector(Sphere.W)));
 	for (int32 Idx : TempInstances)
@@ -1439,7 +2059,7 @@ bool FFoliageMeshInfo::CheckForOverlappingSphere(const FSphere& Sphere)
 }
 
 // Returns whether or not there is are any instances overlapping the instance specified, excluding the set of instances provided
-bool FFoliageMeshInfo::CheckForOverlappingInstanceExcluding(int32 TestInstanceIdx, float Radius, TSet<int32>& ExcludeInstances)
+bool FFoliageInfo::CheckForOverlappingInstanceExcluding(int32 TestInstanceIdx, float Radius, TSet<int32>& ExcludeInstances)
 {
 	FSphere Sphere(Instances[TestInstanceIdx].Location, Radius);
 
@@ -1454,38 +2074,36 @@ bool FFoliageMeshInfo::CheckForOverlappingInstanceExcluding(int32 TestInstanceId
 	return false;
 }
 
-void FFoliageMeshInfo::SelectInstances(AInstancedFoliageActor* InIFA, bool bSelect)
+void FFoliageInfo::SelectInstances(AInstancedFoliageActor* InIFA, bool bSelect)
 {
-	if (Component)
+	if (Implementation->IsInitialized())
 	{
 		InIFA->Modify();
 
 		if (bSelect)
 		{
-			SelectedIndices.Reserve(Component->PerInstanceSMData.Num());
+			SelectedIndices.Reserve(Instances.Num());
 
-			for (int32 i = 0; i < Component->PerInstanceSMData.Num(); ++i)
+			for (int32 i = 0; i < Instances.Num(); ++i)
 			{
 				SelectedIndices.Add(i);
 			}
 
-			Component->SelectInstance(true, 0, Component->PerInstanceSMData.Num());
+			Implementation->SelectInstances(true, 0, SelectedIndices.Num());
 		}
 		else
 		{
+			Implementation->ClearSelection(SelectedIndices);
 			SelectedIndices.Empty();
-			Component->ClearInstanceSelection();
 		}
-
-		Component->MarkRenderStateDirty();
 	}
 }
 
-void FFoliageMeshInfo::SelectInstances(AInstancedFoliageActor* InIFA, bool bSelect, TArray<int32>& InInstances)
+void FFoliageInfo::SelectInstances(AInstancedFoliageActor* InIFA, bool bSelect, TArray<int32>& InInstances)
 {
 	if (InInstances.Num())
 	{
-		check(Component);
+		check(Implementation->IsInitialized());
 		if (bSelect)
 		{
 			InIFA->Modify();
@@ -1495,10 +2113,8 @@ void FFoliageMeshInfo::SelectInstances(AInstancedFoliageActor* InIFA, bool bSele
 			for (int32 i : InInstances)
 			{
 				SelectedIndices.Add(i);
-				Component->SelectInstance(true, i, 1);
+				Implementation->SelectInstances(true, i, 1);
 			}
-
-			Component->MarkRenderStateDirty();
 		}
 		else
 		{
@@ -1509,19 +2125,57 @@ void FFoliageMeshInfo::SelectInstances(AInstancedFoliageActor* InIFA, bool bSele
 				SelectedIndices.Remove(i);
 			}
 
-			if (Component->SelectedInstances.Num() > 0)
+			for (int32 i : InInstances)
 			{
-				for (int32 i : InInstances)
-				{
-					Component->SelectInstance(false, i, 1);
-				}
-
-				Component->MarkRenderStateDirty();
+				Implementation->SelectInstances(false, i, 1);
 			}
 		}
 	}
 }
 
+void FFoliageInfo::Refresh(AInstancedFoliageActor* IFA, bool Async, bool Force)
+{
+	check(Implementation.IsValid())
+	Implementation->Refresh(IFA, this->Instances, Async, Force);
+}
+
+void FFoliageInfo::OnHiddenEditorViewMaskChanged(uint64 InHiddenEditorViews)
+{
+	Implementation->OnHiddenEditorViewMaskChanged(InHiddenEditorViews);
+}
+
+void FFoliageInfo::PreEditUndo(AInstancedFoliageActor* IFA, UFoliageType* FoliageType)
+{
+	Implementation->PreEditUndo(IFA, FoliageType);
+}
+
+void FFoliageInfo::PostEditUndo(AInstancedFoliageActor* IFA, UFoliageType* FoliageType)
+{
+	Implementation->PostEditUndo(IFA, FoliageType, Instances, SelectedIndices);
+		
+	// Regenerate instance hash
+	// We regenerate it here instead of saving to transaction buffer to speed up modify operations
+	InstanceHash->Empty();
+	for (int32 InstanceIdx = 0; InstanceIdx < Instances.Num(); InstanceIdx++)
+	{
+		InstanceHash->InsertInstance(Instances[InstanceIdx].Location, InstanceIdx);
+	}
+}
+
+void FFoliageInfo::EnterEditMode()
+{
+	Implementation->EnterEditMode();
+}
+
+void FFoliageInfo::ExitEditMode()
+{
+	Implementation->ExitEditMode();
+}
+
+TArray<int32> FFoliageInfo::GetInstancesOverlappingBox(const FBox& Box) const
+{
+	return InstanceHash->GetInstancesOverlappingBox(Box);
+}
 #endif	//WITH_EDITOR
 
 //
@@ -1541,6 +2195,11 @@ AInstancedFoliageActor::AInstancedFoliageActor(const FObjectInitializer& ObjectI
 	PrimaryActorTick.bCanEverTick = false;
 }
 
+bool AInstancedFoliageActor::IsOwnedByFoliage(const AActor* InActor)
+{
+	return InActor != nullptr && InActor->ActorHasTag(FoliageActorTag);
+}
+
 AInstancedFoliageActor* AInstancedFoliageActor::GetInstancedFoliageActorForCurrentLevel(UWorld* InWorld, bool bCreateIfNone)
 {
 	return GetInstancedFoliageActorForLevel(InWorld->GetCurrentLevel(), bCreateIfNone);
@@ -1568,12 +2227,9 @@ AInstancedFoliageActor* AInstancedFoliageActor::GetInstancedFoliageActorForLevel
 
 int32 AInstancedFoliageActor::GetOverlappingSphereCount(const UFoliageType* FoliageType, const FSphere& Sphere) const
 {
-	if (const FFoliageMeshInfo* MeshInfo = FindMesh(FoliageType))
+	if (const FFoliageInfo* Info = FindInfo(FoliageType))
 	{
-		if (MeshInfo->Component && MeshInfo->Component->IsTreeFullyBuilt())
-		{
-			return MeshInfo->Component->GetOverlappingSphereCount(Sphere);
-		}
+		return Info->GetOverlappingSphereCount(Sphere);
 	}
 
 	return 0;
@@ -1582,12 +2238,9 @@ int32 AInstancedFoliageActor::GetOverlappingSphereCount(const UFoliageType* Foli
 
 int32 AInstancedFoliageActor::GetOverlappingBoxCount(const UFoliageType* FoliageType, const FBox& Box) const
 {
-	if (const FFoliageMeshInfo* MeshInfo = FindMesh(FoliageType))
+	if (const FFoliageInfo* Info = FindInfo(FoliageType))
 	{
-		if (MeshInfo->Component && MeshInfo->Component->IsTreeFullyBuilt())
-		{
-			return MeshInfo->Component->GetOverlappingBoxCount(Box);
-		}
+		return Info->GetOverlappingBoxCount(Box);
 	}
 
 	return 0;
@@ -1596,67 +2249,58 @@ int32 AInstancedFoliageActor::GetOverlappingBoxCount(const UFoliageType* Foliage
 
 void AInstancedFoliageActor::GetOverlappingBoxTransforms(const UFoliageType* FoliageType, const FBox& Box, TArray<FTransform>& OutTransforms) const
 {
-	if (const FFoliageMeshInfo* MeshInfo = FindMesh(FoliageType))
+	if (const FFoliageInfo* Info = FindInfo(FoliageType))
 	{
-		if (MeshInfo->Component && MeshInfo->Component->IsTreeFullyBuilt())
-		{
-			MeshInfo->Component->GetOverlappingBoxTransforms(Box, OutTransforms);
-		}
+		Info->GetOverlappingBoxTransforms(Box, OutTransforms);
 	}
 }
 
 void AInstancedFoliageActor::GetOverlappingMeshCounts(const FSphere& Sphere, TMap<UStaticMesh*, int32>& OutCounts) const
 {
-	for (auto& MeshPair : FoliageMeshes)
+	for (auto& Pair : FoliageInfos)
 	{
-		FFoliageMeshInfo const* MeshInfo = &*MeshPair.Value;
+		FFoliageInfo const* Info = &*Pair.Value;
 
-		if (MeshInfo && MeshInfo->Component)
+		if (Info)
 		{
-			int32 const Count = MeshInfo->Component->GetOverlappingSphereCount(Sphere);
-			if (Count > 0)
-			{
-				UStaticMesh* const Mesh = MeshInfo->Component->GetStaticMesh();
-				int32& StoredCount = OutCounts.FindOrAdd(Mesh);
-				StoredCount += Count;
-			}
+			Info->GetOverlappingMeshCount(Sphere, OutCounts);
 		}
 	}
 }
 
 
 
-UFoliageType* AInstancedFoliageActor::GetLocalFoliageTypeForMesh(const UStaticMesh* InMesh, FFoliageMeshInfo** OutMeshInfo)
+UFoliageType* AInstancedFoliageActor::GetLocalFoliageTypeForSource(const UObject* InSource, FFoliageInfo** OutMeshInfo)
 {
 	UFoliageType* ReturnType = nullptr;
-	FFoliageMeshInfo* MeshInfo = nullptr;
+	FFoliageInfo* Info = nullptr;
 
-	for (auto& MeshPair : FoliageMeshes)
+	for (auto& Pair : FoliageInfos)
 	{
-		UFoliageType* FoliageType = MeshPair.Key;
+		UFoliageType* FoliageType = Pair.Key;
 		// Check that the type is neither an asset nor blueprint instance
-		if (FoliageType && FoliageType->GetStaticMesh() == InMesh && !FoliageType->IsAsset() && !FoliageType->GetClass()->ClassGeneratedBy)
+		if (FoliageType && FoliageType->GetSource() == InSource && !FoliageType->IsAsset() && !FoliageType->GetClass()->ClassGeneratedBy)
 		{
 			ReturnType = FoliageType;
-			MeshInfo = &*MeshPair.Value;
+			Info = &*Pair.Value;
 			break;
 		}
 	}
 
 	if (OutMeshInfo)
 	{
-		*OutMeshInfo = MeshInfo;
+		*OutMeshInfo = Info;
 	}
 
 	return ReturnType;
 }
 
-void AInstancedFoliageActor::GetAllFoliageTypesForMesh(const UStaticMesh* InMesh, TArray<const UFoliageType*>& OutFoliageTypes)
+void AInstancedFoliageActor::GetAllFoliageTypesForSource(const UObject* InSource, TArray<const UFoliageType*>& OutFoliageTypes)
 {
-	for (auto& MeshPair : FoliageMeshes)
+	for (auto& Pair : FoliageInfos)
 	{
-		UFoliageType* FoliageType = MeshPair.Key;
-		if (FoliageType && FoliageType->GetStaticMesh() == InMesh)
+		UFoliageType* FoliageType = Pair.Key;
+		if (FoliageType && FoliageType->GetSource() == InSource)
 		{
 			OutFoliageTypes.Add(FoliageType);
 		}
@@ -1664,35 +2308,35 @@ void AInstancedFoliageActor::GetAllFoliageTypesForMesh(const UStaticMesh* InMesh
 }
 
 
-FFoliageMeshInfo* AInstancedFoliageActor::FindFoliageTypeOfClass(TSubclassOf<UFoliageType_InstancedStaticMesh> Class)
+FFoliageInfo* AInstancedFoliageActor::FindFoliageTypeOfClass(TSubclassOf<UFoliageType_InstancedStaticMesh> Class)
 {
-	FFoliageMeshInfo* MeshInfo = nullptr;
+	FFoliageInfo* Info = nullptr;
 
-	for (auto& MeshPair : FoliageMeshes)
+	for (auto& Pair : FoliageInfos)
 	{
-		UFoliageType* FoliageType = MeshPair.Key;
+		UFoliageType* FoliageType = Pair.Key;
 		if (FoliageType && FoliageType->GetClass() == Class)
 		{
-			MeshInfo = &MeshPair.Value.Get();
+			Info = &Pair.Value.Get();
 			break;
 		}
 	}
 
-	return MeshInfo;
+	return Info;
 }
 
-FFoliageMeshInfo* AInstancedFoliageActor::FindMesh(const UFoliageType* InType)
+FFoliageInfo* AInstancedFoliageActor::FindInfo(const UFoliageType* InType)
 {
-	TUniqueObj<FFoliageMeshInfo>* MeshInfoEntry = FoliageMeshes.Find(InType);
-	FFoliageMeshInfo* MeshInfo = MeshInfoEntry ? &MeshInfoEntry->Get() : nullptr;
-	return MeshInfo;
+	TUniqueObj<FFoliageInfo>* InfoEntry = FoliageInfos.Find(InType);
+	FFoliageInfo* Info = InfoEntry ? &InfoEntry->Get() : nullptr;
+	return Info;
 }
 
-const FFoliageMeshInfo* AInstancedFoliageActor::FindMesh(const UFoliageType* InType) const
+const FFoliageInfo* AInstancedFoliageActor::FindInfo(const UFoliageType* InType) const
 {
-	const TUniqueObj<FFoliageMeshInfo>* MeshInfoEntry = FoliageMeshes.Find(InType);
-	const FFoliageMeshInfo* MeshInfo = MeshInfoEntry ? &MeshInfoEntry->Get() : nullptr;
-	return MeshInfo;
+	const TUniqueObj<FFoliageInfo>* InfoEntry = FoliageInfos.Find(InType);
+	const FFoliageInfo* Info = InfoEntry ? &InfoEntry->Get() : nullptr;
+	return Info;
 }
 
 
@@ -1723,10 +2367,10 @@ void AInstancedFoliageActor::MoveInstancesForMovedComponent(UActorComponent* InC
 		FRotationMatrix(NewBaseInfo.CachedRotation) *
 		FTranslationMatrix(NewBaseInfo.CachedLocation);
 
-	for (auto& MeshPair : FoliageMeshes)
+	for (auto& Pair : FoliageInfos)
 	{
-		FFoliageMeshInfo& MeshInfo = *MeshPair.Value;
-		const auto* InstanceSet = MeshInfo.ComponentHash.Find(BaseId);
+		FFoliageInfo& Info = *Pair.Value;
+		const auto* InstanceSet = Info.ComponentHash.Find(BaseId);
 		if (InstanceSet && InstanceSet->Num())
 		{
 			if (bFirst)
@@ -1735,14 +2379,13 @@ void AInstancedFoliageActor::MoveInstancesForMovedComponent(UActorComponent* InC
 				Modify();
 			}
 
-			check(MeshInfo.Component);
-			MeshInfo.Component->bAutoRebuildTreeOnInstanceChanges = false;
+			Info.Implementation->BeginUpdate();
 
 			for (int32 InstanceIndex : *InstanceSet)
 			{
-				FFoliageInstance& Instance = MeshInfo.Instances[InstanceIndex];
+				FFoliageInstance& Instance = Info.Instances[InstanceIndex];
 
-				MeshInfo.InstanceHash->RemoveInstance(Instance.Location, InstanceIndex);
+				Info.InstanceHash->RemoveInstance(Instance.Location, InstanceIndex);
 
 				// Apply change
 				FMatrix NewTransform =
@@ -1755,14 +2398,14 @@ void AInstancedFoliageActor::MoveInstancesForMovedComponent(UActorComponent* InC
 				Instance.Rotation = NewTransform.Rotator();
 
 				// Apply render data
-				MeshInfo.Component->UpdateInstanceTransform(InstanceIndex, Instance.GetInstanceWorldTransform(), true, true);
+				Info.Implementation->SetInstanceWorldTransform(InstanceIndex, Instance.GetInstanceWorldTransform(), true);
 
 				// Re-add the new instance location to the hash
-				MeshInfo.InstanceHash->InsertInstance(Instance.Location, InstanceIndex);
+				Info.InstanceHash->InsertInstance(Instance.Location, InstanceIndex);
 			}
 
-			MeshInfo.Component->bAutoRebuildTreeOnInstanceChanges = true;
-			MeshInfo.Component->BuildTreeIfOutdated(true, false);
+			Info.Implementation->EndUpdate();
+			Info.Refresh(this, true, false);
 		}
 	}
 }
@@ -1776,13 +2419,13 @@ void AInstancedFoliageActor::DeleteInstancesForComponent(UActorComponent* InComp
 		return;
 	}
 
-	for (auto& MeshPair : FoliageMeshes)
+	for (auto& Pair : FoliageInfos)
 	{
-		FFoliageMeshInfo& MeshInfo = *MeshPair.Value;
-		const auto* InstanceSet = MeshInfo.ComponentHash.Find(BaseId);
+		FFoliageInfo& Info = *Pair.Value;
+		const auto* InstanceSet = Info.ComponentHash.Find(BaseId);
 		if (InstanceSet)
 		{
-			MeshInfo.RemoveInstances(this, InstanceSet->Array(), true);
+			Info.RemoveInstances(this, InstanceSet->Array(), true);
 		}
 	}
 }
@@ -1796,13 +2439,13 @@ void AInstancedFoliageActor::DeleteInstancesForComponent(UActorComponent* InComp
 		return;
 	}
 
-	FFoliageMeshInfo* MeshInfo = FindMesh(FoliageType);
-	if (MeshInfo)
+	FFoliageInfo* Info = FindInfo(FoliageType);
+	if (Info)
 	{
-		const auto* InstanceSet = MeshInfo->ComponentHash.Find(BaseId);
+		const auto* InstanceSet = Info->ComponentHash.Find(BaseId);
 		if (InstanceSet)
 		{
-			MeshInfo->RemoveInstances(this, InstanceSet->Array(), true);
+			Info->RemoveInstances(this, InstanceSet->Array(), true);
 		}
 	}
 }
@@ -1820,13 +2463,13 @@ void AInstancedFoliageActor::DeleteInstancesForComponent(UWorld* InWorld, UActor
 void AInstancedFoliageActor::DeleteInstancesForProceduralFoliageComponent(const UProceduralFoliageComponent* ProceduralFoliageComponent, bool InRebuildTree)
 {
 	const FGuid& ProceduralGuid = ProceduralFoliageComponent->GetProceduralGuid();
-	for (auto& MeshPair : FoliageMeshes)
+	for (auto& Pair : FoliageInfos)
 	{
-		FFoliageMeshInfo& MeshInfo = *MeshPair.Value;
+		FFoliageInfo& Info = *Pair.Value;
 		TArray<int32> InstancesToRemove;
-		for (int32 InstanceIdx = 0; InstanceIdx < MeshInfo.Instances.Num(); InstanceIdx++)
+		for (int32 InstanceIdx = 0; InstanceIdx < Info.Instances.Num(); InstanceIdx++)
 		{
-			if (MeshInfo.Instances[InstanceIdx].ProceduralGuid == ProceduralGuid)
+			if (Info.Instances[InstanceIdx].ProceduralGuid == ProceduralGuid)
 			{
 				InstancesToRemove.Add(InstanceIdx);
 			}
@@ -1834,7 +2477,7 @@ void AInstancedFoliageActor::DeleteInstancesForProceduralFoliageComponent(const
 
 		if (InstancesToRemove.Num())
 		{
-			MeshInfo.RemoveInstances(this, InstancesToRemove, InRebuildTree);
+			Info.RemoveInstances(this, InstancesToRemove, InRebuildTree);
 		}
 	}
 
@@ -1845,13 +2488,13 @@ void AInstancedFoliageActor::DeleteInstancesForProceduralFoliageComponent(const
 bool AInstancedFoliageActor::ContainsInstancesFromProceduralFoliageComponent(const UProceduralFoliageComponent* ProceduralFoliageComponent)
 {
 	const FGuid& ProceduralGuid = ProceduralFoliageComponent->GetProceduralGuid();
-	for (auto& MeshPair : FoliageMeshes)
+	for (auto& Pair : FoliageInfos)
 	{
-		FFoliageMeshInfo& MeshInfo = *MeshPair.Value;
+		FFoliageInfo& Info = *Pair.Value;
 		TArray<int32> InstancesToRemove;
-		for (int32 InstanceIdx = 0; InstanceIdx < MeshInfo.Instances.Num(); InstanceIdx++)
+		for (int32 InstanceIdx = 0; InstanceIdx < Info.Instances.Num(); InstanceIdx++)
 		{
-			if (MeshInfo.Instances[InstanceIdx].ProceduralGuid == ProceduralGuid)
+			if (Info.Instances[InstanceIdx].ProceduralGuid == ProceduralGuid)
 			{
 				// The procedural component is responsible for an instance
 				return true;
@@ -1883,28 +2526,28 @@ void AInstancedFoliageActor::MoveInstancesForComponentToCurrentLevel(UActorCompo
 		{
 			IFA->Modify();
 
-			for (auto& MeshPair : IFA->FoliageMeshes)
+			for (auto& Pair : IFA->FoliageInfos)
 			{
-				FFoliageMeshInfo& MeshInfo = *MeshPair.Value;
-				UFoliageType* FoliageType = MeshPair.Key;
+				FFoliageInfo& Info = *Pair.Value;
+				UFoliageType* FoliageType = Pair.Key;
 
-				const auto* InstanceSet = MeshInfo.ComponentHash.Find(SourceBaseId);
+				const auto* InstanceSet = Info.ComponentHash.Find(SourceBaseId);
 				if (InstanceSet)
 				{
 					// Duplicate the foliage type if it's not shared
-					FFoliageMeshInfo* TargetMeshInfo = nullptr;
+					FFoliageInfo* TargetMeshInfo = nullptr;
 					UFoliageType* TargetFoliageType = NewIFA->AddFoliageType(FoliageType, &TargetMeshInfo);
 
 					// Add the foliage to the new level
 					for (int32 InstanceIndex : *InstanceSet)
 					{
-						TargetMeshInfo->AddInstance(NewIFA, TargetFoliageType, MeshInfo.Instances[InstanceIndex], InComponent, false);
+						TargetMeshInfo->AddInstance(NewIFA, TargetFoliageType, Info.Instances[InstanceIndex], InComponent);
 					}
 
-					TargetMeshInfo->Component->BuildTreeIfOutdated(true, true);
+					TargetMeshInfo->Refresh(NewIFA, true, true);
 
 					// Remove from old level
-					MeshInfo.RemoveInstances(IFA, InstanceSet->Array(), true);
+					Info.RemoveInstances(IFA, InstanceSet->Array(), true);
 				}
 			}
 		}
@@ -1923,32 +2566,26 @@ void AInstancedFoliageActor::MoveInstancesToNewComponent(UPrimitiveComponent* In
 	AInstancedFoliageActor* TargetIFA = AInstancedFoliageActor::GetInstancedFoliageActorForLevel(InNewComponent->GetTypedOuter<ULevel>(), true);
 	TArray<int32> InstancesToMove;
 
-	for (auto& MeshPair : FoliageMeshes)
+	for (auto& Pair : FoliageInfos)
 	{
-		FFoliageMeshInfo& MeshInfo = *MeshPair.Value;
+		FFoliageInfo& Info = *Pair.Value;
 		
-		if (MeshInfo.Component != nullptr)
-		{
-			InstancesToMove = MeshInfo.Component->GetInstancesOverlappingBox(InBoxWithInstancesToMove);
-		}
-
-		FFoliageMeshInfo* TargetMeshInfo = nullptr;
-		UFoliageType* TargetFoliageType = TargetIFA->AddFoliageType(MeshPair.Key, &TargetMeshInfo);
+		InstancesToMove = Info.GetInstancesOverlappingBox(InBoxWithInstancesToMove);
+		
+		FFoliageInfo* TargetMeshInfo = nullptr;
+		UFoliageType* TargetFoliageType = TargetIFA->AddFoliageType(Pair.Key, &TargetMeshInfo);
 
 		// Add the foliage to the new level
 		for (int32 InstanceIndex : InstancesToMove)
 		{
-			FFoliageInstance NewInstance = MeshInfo.Instances[InstanceIndex];
-			TargetMeshInfo->AddInstance(TargetIFA, TargetFoliageType, NewInstance, InNewComponent, false);
-		}
-
-		if (TargetMeshInfo->Component != nullptr)
-		{
-			TargetMeshInfo->Component->BuildTreeIfOutdated(true, true);
+			FFoliageInstance NewInstance = Info.Instances[InstanceIndex];
+			TargetMeshInfo->AddInstance(TargetIFA, TargetFoliageType, NewInstance, InNewComponent);
 		}
 
+		TargetMeshInfo->Refresh(TargetIFA, true, true);
+		
 		// Remove from old level
-		MeshInfo.RemoveInstances(this, InstancesToMove, true);
+		Info.RemoveInstances(this, InstancesToMove, true);
 	}
 }
 
@@ -1965,12 +2602,12 @@ void AInstancedFoliageActor::MoveInstancesToNewComponent(UPrimitiveComponent* In
 
 	const auto NewBaseId = TargetIFA->InstanceBaseCache.AddInstanceBaseId(InNewComponent);
 
-	for (auto& MeshPair : FoliageMeshes)
+	for (auto& Pair : FoliageInfos)
 	{
-		FFoliageMeshInfo& MeshInfo = *MeshPair.Value;
+		FFoliageInfo& Info = *Pair.Value;
 
 		TSet<int32> InstanceSet;
-		if (MeshInfo.ComponentHash.RemoveAndCopyValue(OldBaseId, InstanceSet) && InstanceSet.Num())
+		if (Info.ComponentHash.RemoveAndCopyValue(OldBaseId, InstanceSet) && InstanceSet.Num())
 		{
 			// For same FoliageActor can just remap the instances, otherwise we have to do a more complex move
 			if (TargetIFA == this)
@@ -1978,32 +2615,29 @@ void AInstancedFoliageActor::MoveInstancesToNewComponent(UPrimitiveComponent* In
 				// Update the instances
 				for (int32 InstanceIndex : InstanceSet)
 				{
-					MeshInfo.Instances[InstanceIndex].BaseId = NewBaseId;
+					Info.Instances[InstanceIndex].BaseId = NewBaseId;
 				}
 
 				// Update the hash
-				MeshInfo.ComponentHash.Add(NewBaseId, MoveTemp(InstanceSet));
+				Info.ComponentHash.Add(NewBaseId, MoveTemp(InstanceSet));
 			}
 			else
 			{
-				FFoliageMeshInfo* TargetMeshInfo = nullptr;
-				UFoliageType* TargetFoliageType = TargetIFA->AddFoliageType(MeshPair.Key, &TargetMeshInfo);
+				FFoliageInfo* TargetMeshInfo = nullptr;
+				UFoliageType* TargetFoliageType = TargetIFA->AddFoliageType(Pair.Key, &TargetMeshInfo);
 
 				// Add the foliage to the new level
 				for (int32 InstanceIndex : InstanceSet)
 				{
-					FFoliageInstance NewInstance = MeshInfo.Instances[InstanceIndex];
+					FFoliageInstance NewInstance = Info.Instances[InstanceIndex];
 					NewInstance.BaseId = NewBaseId;
-					TargetMeshInfo->AddInstance(TargetIFA, TargetFoliageType, NewInstance, false);
+					TargetMeshInfo->AddInstance(TargetIFA, TargetFoliageType, NewInstance);
 				}
 
-				if (TargetMeshInfo->Component != nullptr)
-				{
-					TargetMeshInfo->Component->BuildTreeIfOutdated(true, true);
-				}
+				TargetMeshInfo->Refresh(TargetIFA, true, true);
 
 				// Remove from old level
-				MeshInfo.RemoveInstances(this, InstanceSet.Array(), true);
+				Info.RemoveInstances(this, InstanceSet.Array(), true);
 			}
 		}
 	}
@@ -2027,7 +2661,7 @@ void AInstancedFoliageActor::MoveInstancesToNewComponent(UWorld* InWorld, UPrimi
 	}
 }
 
-void AInstancedFoliageActor::MoveInstancesToLevel(ULevel* InTargetLevel, TSet<int32>& InInstanceList, FFoliageMeshInfo* InCurrentMeshInfo, UFoliageType* InFoliageType)
+void AInstancedFoliageActor::MoveInstancesToLevel(ULevel* InTargetLevel, TSet<int32>& InInstanceList, FFoliageInfo* InCurrentMeshInfo, UFoliageType* InFoliageType)
 {
 	if (InTargetLevel == GetLevel())
 	{
@@ -2040,20 +2674,17 @@ void AInstancedFoliageActor::MoveInstancesToLevel(ULevel* InTargetLevel, TSet<in
 	TargetIFA->Modify();
 
 	// Do move
-	FFoliageMeshInfo* TargetMeshInfo = nullptr;
+	FFoliageInfo* TargetMeshInfo = nullptr;
 	UFoliageType* TargetFoliageType = TargetIFA->AddFoliageType(InFoliageType, &TargetMeshInfo);
 
 	// Add selected instances to the target actor
 	for (int32 InstanceIndex : InInstanceList)
 	{
 		FFoliageInstance& Instance = InCurrentMeshInfo->Instances[InstanceIndex];
-		TargetMeshInfo->AddInstance(TargetIFA, TargetFoliageType, Instance, InstanceBaseCache.GetInstanceBasePtr(Instance.BaseId).Get(), false);
+		TargetMeshInfo->AddInstance(TargetIFA, TargetFoliageType, Instance, InstanceBaseCache.GetInstanceBasePtr(Instance.BaseId).Get());
 	}
 
-	if (TargetMeshInfo->Component != nullptr)
-	{
-		TargetMeshInfo->Component->BuildTreeIfOutdated(true, true);
-	}
+	TargetMeshInfo->Refresh(TargetIFA, true, true);
 
 	// Remove selected instances from this actor
 	InCurrentMeshInfo->RemoveInstances(this, InInstanceList.Array(), true);
@@ -2066,12 +2697,12 @@ void AInstancedFoliageActor::MoveSelectedInstancesToLevel(ULevel* InTargetLevel)
 		return;
 	}
 
-	for (auto& MeshPair : FoliageMeshes)
+	for (auto& Pair : FoliageInfos)
 	{
-		FFoliageMeshInfo& MeshInfo = *MeshPair.Value;
-		UFoliageType* FoliageType = MeshPair.Key;
+		FFoliageInfo& Info = *Pair.Value;
+		UFoliageType* FoliageType = Pair.Key;
 
-		MoveInstancesToLevel(InTargetLevel, MeshInfo.SelectedIndices, &MeshInfo, FoliageType);
+		MoveInstancesToLevel(InTargetLevel, Info.SelectedIndices, &Info, FoliageType);
 	}
 }
 
@@ -2082,19 +2713,19 @@ void AInstancedFoliageActor::MoveAllInstancesToLevel(ULevel* InTargetLevel)
 		return;
 	}
 
-	for (auto& MeshPair : FoliageMeshes)
+	for (auto& Pair : FoliageInfos)
 	{
-		FFoliageMeshInfo& MeshInfo = *MeshPair.Value;
-		UFoliageType* FoliageType = MeshPair.Key;
+		FFoliageInfo& Info = *Pair.Value;
+		UFoliageType* FoliageType = Pair.Key;
 
 		TSet<int32> instancesList;
 
-		for (int32 i = 0; i < MeshInfo.Instances.Num(); ++i)
+		for (int32 i = 0; i < Info.Instances.Num(); ++i)
 		{
 			instancesList.Add(i);
 		}
 
-		MoveInstancesToLevel(InTargetLevel, instancesList, &MeshInfo, FoliageType);
+		MoveInstancesToLevel(InTargetLevel, instancesList, &Info, FoliageType);
 	}
 }
 
@@ -2105,18 +2736,18 @@ TMap<UFoliageType*, TArray<const FFoliageInstancePlacementInfo*>> AInstancedFoli
 
 	if (BaseId != FFoliageInstanceBaseCache::InvalidBaseId)
 	{
-		for (auto& MeshPair : FoliageMeshes)
+		for (auto& Pair : FoliageInfos)
 		{
-			const FFoliageMeshInfo& MeshInfo = *MeshPair.Value;
-			const auto* InstanceSet = MeshInfo.ComponentHash.Find(BaseId);
+			const FFoliageInfo& Info = *Pair.Value;
+			const auto* InstanceSet = Info.ComponentHash.Find(BaseId);
 			if (InstanceSet)
 			{
-				TArray<const FFoliageInstancePlacementInfo*>& Array = Result.Add(MeshPair.Key, TArray<const FFoliageInstancePlacementInfo*>());
+				TArray<const FFoliageInstancePlacementInfo*>& Array = Result.Add(Pair.Key, TArray<const FFoliageInstancePlacementInfo*>());
 				Array.Empty(InstanceSet->Num());
 
 				for (int32 InstanceIndex : *InstanceSet)
 				{
-					const FFoliageInstancePlacementInfo* Instance = &MeshInfo.Instances[InstanceIndex];
+					const FFoliageInstancePlacementInfo* Instance = &Info.Instances[InstanceIndex];
 					Array.Add(Instance);
 				}
 			}
@@ -2126,103 +2757,76 @@ TMap<UFoliageType*, TArray<const FFoliageInstancePlacementInfo*>> AInstancedFoli
 	return Result;
 }
 
-FFoliageMeshInfo* AInstancedFoliageActor::FindOrAddMesh(UFoliageType* InType)
-{
-	TUniqueObj<FFoliageMeshInfo>* MeshInfoEntry = FoliageMeshes.Find(InType);
-	FFoliageMeshInfo* MeshInfo = MeshInfoEntry ? &MeshInfoEntry->Get() : AddMesh(InType);
-	return MeshInfo;
-}
-
-
-void UpdateSettingsBounds(const UStaticMesh* InMesh, UFoliageType_InstancedStaticMesh* Settings)
+FFoliageInfo* AInstancedFoliageActor::FindOrAddMesh(UFoliageType* InType)
 {
-	const FBoxSphereBounds MeshBounds = InMesh->GetBounds();
-
-	Settings->MeshBounds = MeshBounds;
-
-	// Make bottom only bound
-	FBox LowBound = MeshBounds.GetBox();
-	LowBound.Max.Z = LowBound.Min.Z + (LowBound.Max.Z - LowBound.Min.Z) * 0.1f;
-
-	float MinX = FLT_MAX, MaxX = FLT_MIN, MinY = FLT_MAX, MaxY = FLT_MIN;
-	Settings->LowBoundOriginRadius = FVector::ZeroVector;
-
-	if (InMesh->RenderData)
-	{
-		FPositionVertexBuffer& PositionVertexBuffer = InMesh->RenderData->LODResources[0].VertexBuffers.PositionVertexBuffer;
-		for (uint32 Index = 0; Index < PositionVertexBuffer.GetNumVertices(); ++Index)
-		{
-			const FVector& Pos = PositionVertexBuffer.VertexPosition(Index);
-			if (Pos.Z < LowBound.Max.Z)
-			{
-				MinX = FMath::Min(MinX, Pos.X);
-				MinY = FMath::Min(MinY, Pos.Y);
-				MaxX = FMath::Max(MaxX, Pos.X);
-				MaxY = FMath::Max(MaxY, Pos.Y);
-			}
-		}
-	}
-
-	Settings->LowBoundOriginRadius = FVector((MinX + MaxX), (MinY + MaxY), FMath::Sqrt(FMath::Square(MaxX - MinX) + FMath::Square(MaxY - MinY))) * 0.5f;
+	TUniqueObj<FFoliageInfo>* MeshInfoEntry = FoliageInfos.Find(InType);
+	FFoliageInfo* Info = MeshInfoEntry ? &MeshInfoEntry->Get() : AddMesh(InType);
+	return Info;
 }
 
-UFoliageType* AInstancedFoliageActor::AddFoliageType(const UFoliageType* InType, FFoliageMeshInfo** OutInfo)
+UFoliageType* AInstancedFoliageActor::AddFoliageType(const UFoliageType* InType, FFoliageInfo** OutInfo)
 {
-	FFoliageMeshInfo* MeshInfo = nullptr;
+	FFoliageInfo* Info = nullptr;
 	UFoliageType* FoliageType = const_cast<UFoliageType*>(InType);
 
 	if (FoliageType->GetOuter() == this || FoliageType->IsAsset())
 	{
-		auto ExistingMeshInfo = FoliageMeshes.Find(FoliageType);
+		auto ExistingMeshInfo = FoliageInfos.Find(FoliageType);
 		if (!ExistingMeshInfo)
 		{
 			Modify();
-			MeshInfo = &FoliageMeshes.Add(FoliageType).Get();
+			Info = &FoliageInfos.Add(FoliageType).Get();
 		}
 		else
 		{
-			MeshInfo = &ExistingMeshInfo->Get();
+			Info = &ExistingMeshInfo->Get();
 		}
 	}
 	else if (FoliageType->GetClass()->ClassGeneratedBy)
 	{
 		// Foliage type blueprint
-		FFoliageMeshInfo* ExistingMeshInfo = FindFoliageTypeOfClass(FoliageType->GetClass());
+		FFoliageInfo* ExistingMeshInfo = FindFoliageTypeOfClass(FoliageType->GetClass());
 		if (!ExistingMeshInfo)
 		{
 			Modify();
 			FoliageType = DuplicateObject<UFoliageType>(InType, this);
-			MeshInfo = &FoliageMeshes.Add(FoliageType).Get();
+			Info = &FoliageInfos.Add(FoliageType).Get();
 		}
 		else
 		{
-			MeshInfo = ExistingMeshInfo;
+			Info = ExistingMeshInfo;
 		}
 	}
 	else
 	{
 		// Unique meshes only
 		// Multiple entries for same static mesh can be added using FoliageType as an asset
-		FoliageType = GetLocalFoliageTypeForMesh(FoliageType->GetStaticMesh(), &MeshInfo);
+		FoliageType = GetLocalFoliageTypeForSource(FoliageType->GetSource(), &Info);
 		if (FoliageType == nullptr)
 		{
 			Modify();
 			FoliageType = DuplicateObject<UFoliageType>(InType, this);
-			MeshInfo = &FoliageMeshes.Add(FoliageType).Get();
+			Info = &FoliageInfos.Add(FoliageType).Get();
 		}
 	}
 
+	if (Info && !Info->Implementation.IsValid())
+	{
+		Info->CreateImplementation(FoliageType);
+		check(Info->Implementation.IsValid());
+	}
+
 	if (OutInfo)
 	{
-		*OutInfo = MeshInfo;
+		*OutInfo = Info;
 	}
 
 	return FoliageType;
 }
 
-FFoliageMeshInfo* AInstancedFoliageActor::AddMesh(UStaticMesh* InMesh, UFoliageType** OutSettings, const UFoliageType_InstancedStaticMesh* DefaultSettings)
+FFoliageInfo* AInstancedFoliageActor::AddMesh(UStaticMesh* InMesh, UFoliageType** OutSettings, const UFoliageType_InstancedStaticMesh* DefaultSettings)
 {
-	check(GetLocalFoliageTypeForMesh(InMesh) == nullptr);
+	check(GetLocalFoliageTypeForSource(InMesh) == nullptr);
 
 	MarkPackageDirty();
 
@@ -2240,30 +2844,28 @@ FFoliageMeshInfo* AInstancedFoliageActor::AddMesh(UStaticMesh* InMesh, UFoliageT
 	{
 		Settings = NewObject<UFoliageType_InstancedStaticMesh>(this, NAME_None, RF_Transactional);
 	}
-	Settings->Mesh = InMesh;
-
-	FFoliageMeshInfo* MeshInfo = AddMesh(Settings);
-	UpdateSettingsBounds(InMesh, Settings);
-
+	Settings->SetStaticMesh(InMesh);
+	FFoliageInfo* Info = AddMesh(Settings);
+	
 	if (OutSettings)
 	{
 		*OutSettings = Settings;
 	}
 
-	return MeshInfo;
+	return Info;
 }
 
-FFoliageMeshInfo* AInstancedFoliageActor::AddMesh(UFoliageType* InType)
+FFoliageInfo* AInstancedFoliageActor::AddMesh(UFoliageType* InType)
 {
-	check(FoliageMeshes.Find(InType) == nullptr);
+	check(FoliageInfos.Find(InType) == nullptr);
 
 	Modify();
 
-	FFoliageMeshInfo* MeshInfo = &*FoliageMeshes.Add(InType);
-	MeshInfo->FoliageTypeUpdateGuid = InType->UpdateGuid;
+	FFoliageInfo* Info = &*FoliageInfos.Add(InType);
+	Info->FoliageTypeUpdateGuid = InType->UpdateGuid;
 	InType->IsSelected = true;
 
-	return MeshInfo;
+	return Info;
 }
 
 void AInstancedFoliageActor::RemoveFoliageType(UFoliageType** InFoliageTypes, int32 Num)
@@ -2275,24 +2877,15 @@ void AInstancedFoliageActor::RemoveFoliageType(UFoliageType** InFoliageTypes, in
 	for (int32 FoliageTypeIdx = 0; FoliageTypeIdx < Num; ++FoliageTypeIdx)
 	{
 		const UFoliageType* FoliageType = InFoliageTypes[FoliageTypeIdx];
-		FFoliageMeshInfo* MeshInfo = FindMesh(FoliageType);
-		if (MeshInfo)
+		FFoliageInfo* Info = FindInfo(FoliageType);
+		if (Info)
 		{
-			if (MeshInfo->Component)
+			if (Info->IsInitialized())
 			{
-				if (MeshInfo->Component->GetStaticMesh() != nullptr)
-				{
-					MeshInfo->Component->GetStaticMesh()->GetOnExtendedBoundsChanged().RemoveAll(MeshInfo);
-				}
-
-				MeshInfo->Component->ClearInstances();
-				MeshInfo->Component->SetFlags(RF_Transactional);
-				MeshInfo->Component->Modify();
-				MeshInfo->Component->DestroyComponent();
-				MeshInfo->Component = nullptr;
+				Info->Uninitialize();
 			}
 
-			FoliageMeshes.Remove(FoliageType);
+			FoliageInfos.Remove(FoliageType);
 		}
 	}
 
@@ -2306,38 +2899,35 @@ void AInstancedFoliageActor::SelectInstance(UInstancedStaticMeshComponent* InCom
 	// If we're not toggling, we need to first deselect everything else
 	if (!bToggle)
 	{
-		for (auto& MeshPair : FoliageMeshes)
+		for (auto& Pair : FoliageInfos)
 		{
-			FFoliageMeshInfo& MeshInfo = *MeshPair.Value;
-
-			if (MeshInfo.Instances.Num() > 0)
-			{
-				check(MeshInfo.Component);
-				MeshInfo.Component->ClearInstanceSelection();
-				MeshInfo.Component->MarkRenderStateDirty();
-				MeshInfo.SelectedIndices.Empty();
-			}
+			FFoliageInfo& Info = *Pair.Value;
+			Info.ClearSelection();
 		}
 	}
 
 	if (InComponent)
 	{
 		UFoliageType* Type = nullptr;
-		FFoliageMeshInfo* MeshInfo = nullptr;
+		FFoliageInfo* Info = nullptr;
 
-		for (auto& MeshPair : FoliageMeshes)
+		for (auto& Pair : FoliageInfos)
 		{
-			if (MeshPair.Value->Component == InComponent)
+			if (Pair.Value->Type == EFoliageImplType::StaticMesh)
 			{
-				Type = MeshPair.Key;
-				MeshInfo = &MeshPair.Value.Get();
-				break;
+				FFoliageStaticMesh* FoliageStaticMesh = StaticCast<FFoliageStaticMesh*>(Pair.Value->Implementation.Get());
+				if (FoliageStaticMesh->Component == InComponent)
+				{
+					Type = Pair.Key;
+					Info = &Pair.Value.Get();
+					break;
+				}
 			}
 		}
 
-		if (MeshInfo)
+		if (Info)
 		{
-			bool bIsSelected = MeshInfo->SelectedIndices.Contains(InInstanceIndex);
+			bool bIsSelected = Info->SelectedIndices.Contains(InInstanceIndex);
 
 			// Deselect if it's already selected.
 			if (InInstanceIndex < InComponent->SelectedInstances.Num())
@@ -2348,7 +2938,7 @@ void AInstancedFoliageActor::SelectInstance(UInstancedStaticMeshComponent* InCom
 
 			if (bIsSelected)
 			{
-				MeshInfo->SelectedIndices.Remove(InInstanceIndex);
+				Info->SelectedIndices.Remove(InInstanceIndex);
 			}
 
 			if (!bToggle || !bIsSelected)
@@ -2357,7 +2947,65 @@ void AInstancedFoliageActor::SelectInstance(UInstancedStaticMeshComponent* InCom
 				InComponent->SelectInstance(true, InInstanceIndex, 1);
 				InComponent->MarkRenderStateDirty();
 
-				MeshInfo->SelectedIndices.Add(InInstanceIndex);
+				Info->SelectedIndices.Add(InInstanceIndex);
+			}
+		}
+	}
+}
+
+void AInstancedFoliageActor::SelectInstance(AActor* InActor, bool bToggle)
+{
+	Modify();
+
+	// If we're not toggling, we need to first deselect everything else
+	if (!bToggle)
+	{
+		for (auto& Pair : FoliageInfos)
+		{
+			FFoliageInfo& Info = *Pair.Value;
+			Info.ClearSelection();
+		}
+	}
+
+	if (InActor)
+	{
+		UFoliageType* Type = nullptr;
+		FFoliageInfo* Info = nullptr;
+		FFoliageActor* FoliageActor = nullptr;
+		int32 index = INDEX_NONE;
+
+		for (auto& Pair : FoliageInfos)
+		{
+			if (Pair.Value->Type == EFoliageImplType::Actor)
+			{
+				FFoliageActor* CurrentFoliageActor = StaticCast<FFoliageActor*>(Pair.Value->Implementation.Get());
+				index = CurrentFoliageActor->FindIndex(InActor);
+				if (index != INDEX_NONE)
+				{
+					Info = &Pair.Value.Get();
+					FoliageActor = CurrentFoliageActor;
+					break;
+				}
+			}
+		}
+
+		if (Info)
+		{
+			bool bIsSelected = Info->SelectedIndices.Contains(index);
+
+			FoliageActor->SelectInstances(false, index, 1);
+			
+			if (bIsSelected)
+			{
+				Info->SelectedIndices.Remove(index);
+			}
+
+			if (!bToggle || !bIsSelected)
+			{
+				// Add the selection
+				FoliageActor->SelectInstances(true, index, 1);
+
+				Info->SelectedIndices.Add(index);
 			}
 		}
 	}
@@ -2365,9 +3013,9 @@ void AInstancedFoliageActor::SelectInstance(UInstancedStaticMeshComponent* InCom
 
 bool AInstancedFoliageActor::HasSelectedInstances() const
 {
-	for (const auto& MeshPair : FoliageMeshes)
+	for (const auto& Pair : FoliageInfos)
 	{
-		if (MeshPair.Value->SelectedIndices.Num() > 0)
+		if (Pair.Value->SelectedIndices.Num() > 0)
 		{
 			return true;
 		}
@@ -2376,27 +3024,27 @@ bool AInstancedFoliageActor::HasSelectedInstances() const
 	return false;
 }
 
-TMap<UFoliageType*, FFoliageMeshInfo*> AInstancedFoliageActor::GetAllInstancesFoliageType()
+TMap<UFoliageType*, FFoliageInfo*> AInstancedFoliageActor::GetAllInstancesFoliageType()
 {
-	TMap<UFoliageType*, FFoliageMeshInfo*> InstanceFoliageTypes;
+	TMap<UFoliageType*, FFoliageInfo*> InstanceFoliageTypes;
 
-	for (auto& MeshPair : FoliageMeshes)
+	for (auto& Pair : FoliageInfos)
 	{
-		InstanceFoliageTypes.Add(MeshPair.Key, &MeshPair.Value.Get());
+		InstanceFoliageTypes.Add(Pair.Key, &Pair.Value.Get());
 	}
 
 	return InstanceFoliageTypes;
 }
 
-TMap<UFoliageType*, FFoliageMeshInfo*> AInstancedFoliageActor::GetSelectedInstancesFoliageType()
+TMap<UFoliageType*, FFoliageInfo*> AInstancedFoliageActor::GetSelectedInstancesFoliageType()
 {
-	TMap<UFoliageType*, FFoliageMeshInfo*> SelectedInstanceFoliageTypes;
+	TMap<UFoliageType*, FFoliageInfo*> SelectedInstanceFoliageTypes;
 
-	for (auto& MeshPair : FoliageMeshes)
+	for (auto& Pair : FoliageInfos)
 	{
-		if (MeshPair.Value->SelectedIndices.Num() > 0)
+		if (Pair.Value->SelectedIndices.Num() > 0)
 		{
-			SelectedInstanceFoliageTypes.Add(MeshPair.Key, &MeshPair.Value.Get());
+			SelectedInstanceFoliageTypes.Add(Pair.Key, &Pair.Value.Get());
 		}
 	}
 
@@ -2407,16 +3055,25 @@ void AInstancedFoliageActor::Destroyed()
 {
 	if (GIsEditor && !GetWorld()->IsGameWorld())
 	{
-		for (auto& MeshPair : FoliageMeshes)
+		for (auto& Pair : FoliageInfos)
 		{
-			UHierarchicalInstancedStaticMeshComponent* Component = MeshPair.Value->Component;
+			if (Pair.Value->Type == EFoliageImplType::StaticMesh)
+			{
+				FFoliageStaticMesh* FoliageStaticMesh = StaticCast<FFoliageStaticMesh*>(Pair.Value->Implementation.Get());
+				UHierarchicalInstancedStaticMeshComponent* Component = FoliageStaticMesh->Component;
 
-			if (Component)
+				if (Component)
+				{
+					Component->ClearInstances();
+					// Save the component's PendingKill flag to restore the component if the delete is undone.
+					Component->SetFlags(RF_Transactional);
+					Component->Modify();
+				}
+			}
+			else if (Pair.Value->Type == EFoliageImplType::Actor)
 			{
-				Component->ClearInstances();
-				// Save the component's PendingKill flag to restore the component if the delete is undone.
-				Component->SetFlags(RF_Transactional);
-				Component->Modify();
+				FFoliageActor* FoliageActor = StaticCast<FFoliageActor*>(Pair.Value->Implementation.Get());
+				FoliageActor->DestroyActors(false);
 			}
 		}
 	}
@@ -2429,14 +3086,11 @@ void AInstancedFoliageActor::PreEditUndo()
 	Super::PreEditUndo();
 
 	// Remove all delegate as we dont know what the Undo will affect and we will simply readd those still valid afterward
-	for (auto& MeshPair : FoliageMeshes)
+	for (auto& Pair : FoliageInfos)
 	{
-		FFoliageMeshInfo& MeshInfo = *MeshPair.Value;
+		FFoliageInfo& Info = *Pair.Value;
 
-		if (MeshPair.Key != nullptr && MeshPair.Key->GetStaticMesh() != nullptr)
-		{
-			MeshPair.Key->GetStaticMesh()->GetOnExtendedBoundsChanged().RemoveAll(&MeshInfo);
-		}
+		Info.PreEditUndo(this, Pair.Key);
 	}
 }
 
@@ -2448,25 +3102,11 @@ void AInstancedFoliageActor::PostEditUndo()
 
 	InstanceBaseCache.UpdateInstanceBaseCachedTransforms();
 
-	for (auto& MeshPair : FoliageMeshes)
+	for (auto& Pair : FoliageInfos)
 	{
-		FFoliageMeshInfo& MeshInfo = *MeshPair.Value;
-
-		if (MeshInfo.Component != nullptr && MeshPair.Key->GetStaticMesh() != nullptr)
-		{
-			MeshPair.Key->GetStaticMesh()->GetOnExtendedBoundsChanged().AddRaw(&MeshInfo, &FFoliageMeshInfo::HandleComponentMeshBoundsChanged);
-		}
+		FFoliageInfo& Info = *Pair.Value;
 
-		MeshInfo.CheckComponentClass(this, MeshPair.Key);
-		MeshInfo.ReapplyInstancesToComponent();
-
-		// Regenerate instance hash
-		// We regenerate it here instead of saving to transaction buffer to speed up modify operations
-		MeshInfo.InstanceHash->Empty();
-		for (int32 InstanceIdx = 0; InstanceIdx < MeshInfo.Instances.Num(); InstanceIdx++)
-		{
-			MeshInfo.InstanceHash->InsertInstance(MeshInfo.Instances[InstanceIdx].Location, InstanceIdx);
-		}
+		Info.PostEditUndo(this, Pair.Key);
 	}
 }
 
@@ -2482,39 +3122,25 @@ bool AInstancedFoliageActor::ShouldImport(FString* ActorPropString, bool IsMovin
 	return false;
 }
 
-void AInstancedFoliageActor::ApplySelectionToComponents(bool bApply)
+void AInstancedFoliageActor::ApplySelection(bool bApply)
 {
-	for (auto& MeshPair : FoliageMeshes)
+	for (auto& Pair : FoliageInfos)
 	{
-		FFoliageMeshInfo& MeshInfo = *MeshPair.Value;
-		UHierarchicalInstancedStaticMeshComponent* Component = MeshInfo.Component;
-
-		if (Component && (bApply || Component->SelectedInstances.Num() > 0))
-		{
-			Component->ClearInstanceSelection();
-
-			if (bApply)
-			{
-				for (int32 i : MeshInfo.SelectedIndices)
-				{
-					Component->SelectInstance(true, i, 1);
-				}
-			}
-
-			Component->MarkRenderStateDirty();
-		}
+		FFoliageInfo& Info = *Pair.Value;
+				
+		Info.Implementation->ApplySelection(bApply, Info.SelectedIndices);
 	}
 }
 
 bool AInstancedFoliageActor::GetSelectionLocation(FVector& OutLocation) const
 {
-	for (const auto& MeshPair : FoliageMeshes)
+	for (const auto& Pair : FoliageInfos)
 	{
-		const FFoliageMeshInfo& MeshInfo = MeshPair.Value.Get();
-		if (MeshInfo.SelectedIndices.Num())
+		const FFoliageInfo& Info = Pair.Value.Get();
+		if (Info.SelectedIndices.Num())
 		{
-			const int32 InstanceIdx = (*MeshInfo.SelectedIndices.CreateConstIterator());
-			OutLocation = MeshInfo.Instances[InstanceIdx].Location;
+			const int32 InstanceIdx = (*Info.SelectedIndices.CreateConstIterator());
+			OutLocation = Info.Instances[InstanceIdx].Location;
 			return true;
 		}
 	}
@@ -2548,14 +3174,14 @@ void AInstancedFoliageActor::MapRebuild()
 	check(World);
 
 	// For each foliage brush, represented by the mesh/info pair
-	for (auto& MeshPair : FoliageMeshes)
+	for (auto& Pair : FoliageInfos)
 	{
 		// each target component has some foliage instances
-		FFoliageMeshInfo const& MeshInfo = *MeshPair.Value;
-		UFoliageType* Settings = MeshPair.Key;
+		FFoliageInfo const& Info = *Pair.Value;
+		UFoliageType* Settings = Pair.Key;
 		check(Settings);
 
-		for (auto& ComponentFoliagePair : MeshInfo.ComponentHash)
+		for (auto& ComponentFoliagePair : Info.ComponentHash)
 		{
 			// BSP components are UModelComponents - they are the only ones we need to change
 			auto BaseComponentPtr = InstanceBaseCache.GetInstanceBasePtr(ComponentFoliagePair.Key);
@@ -2571,8 +3197,8 @@ void AInstancedFoliageActor::MapRebuild()
 				for (int32 InstanceIdx : ComponentFoliagePair.Value)
 				{
 					// Use a line test against the world. This is not very reliable as we don't know the original trace direction.
-					check(MeshInfo.Instances.IsValidIndex(InstanceIdx));
-					FFoliageInstance const& Instance = MeshInfo.Instances[InstanceIdx];
+					check(Info.Instances.IsValidIndex(InstanceIdx));
+					FFoliageInstance const& Instance = Info.Instances[InstanceIdx];
 
 					FFoliageInstance NewInstance = Instance;
 
@@ -2605,13 +3231,13 @@ void AInstancedFoliageActor::MapRebuild()
 	{
 		UFoliageType* Settings = NewInstancePair.Key;
 		check(Settings);
-		FFoliageMeshInfo& MeshInfo = *FindOrAddMesh(Settings);
+		FFoliageInfo& Info = *FindOrAddMesh(Settings);
 		for (FFoliageInstance& Instance : NewInstancePair.Value)
 		{
-			MeshInfo.AddInstance(this, Settings, Instance, false);
+			Info.AddInstance(this, Settings, Instance);
 		}
 
-		MeshInfo.Component->BuildTreeIfOutdated(true, true);
+		Info.Refresh(this, true, true);
 	}
 }
 
@@ -2686,11 +3312,17 @@ void AInstancedFoliageActor::Serialize(FArchive& Ar)
 		{
 #if WITH_EDITORONLY_DATA
 			Ar << FoliageMeshes_Deprecated;
+#endif
+		}
+		else if(Ar.CustomVer(FFoliageCustomVersion::GUID) < FFoliageCustomVersion::FoliageActorSupport)
+		{
+#if WITH_EDITORONLY_DATA
+			Ar << FoliageMeshes_Deprecated2;
 #endif
 		}
 		else
 		{
-			Ar << FoliageMeshes;
+			Ar << FoliageInfos;
 		}
 	}
 
@@ -2708,6 +3340,8 @@ void AInstancedFoliageActor::Serialize(FArchive& Ar)
 }
 
 #if WITH_EDITOR
+AInstancedFoliageActor::FOnSelectionChanged AInstancedFoliageActor::SelectionChanged;
+
 void AInstancedFoliageActor::PostInitProperties()
 {
 	Super::PostInitProperties();
@@ -2768,12 +3402,17 @@ void AInstancedFoliageActor::PostLoad()
 	{
 		if (GetLinkerCustomVersion(FFoliageCustomVersion::GUID) < FFoliageCustomVersion::CrossLevelBase)
 		{
-			ConvertDeprecatedFoliageMeshes(this, FoliageMeshes_Deprecated, FoliageMeshes);
+			ConvertDeprecatedFoliageMeshes(this, FoliageMeshes_Deprecated, FoliageInfos);
 			FoliageMeshes_Deprecated.Empty();
 		}
-
+		else if (GetLinkerCustomVersion(FFoliageCustomVersion::GUID) < FFoliageCustomVersion::FoliageActorSupport)
+		{
+			ConvertDeprecated2FoliageMeshes(this, FoliageMeshes_Deprecated2, FoliageInfos);
+			FoliageMeshes_Deprecated2.Empty();
+		}
+			   
 		{
-			bool bContainsNull = FoliageMeshes.Remove(nullptr) > 0;
+			bool bContainsNull = FoliageInfos.Remove(nullptr) > 0;
 			if (bContainsNull)
 			{
 				FMessageLog("MapCheck").Warning()
@@ -2782,32 +3421,32 @@ void AInstancedFoliageActor::PostLoad()
 					->AddToken(FMapErrorToken::Create(FMapErrors::FoliageMissingStaticMesh));
 				while (bContainsNull)
 				{
-					bContainsNull = FoliageMeshes.Remove(nullptr) > 0;
+					bContainsNull = FoliageInfos.Remove(nullptr) > 0;
 				}
 			}
 		}
 
 		TArray<UFoliageType*> FoliageTypeToRemove;
 
-		for (auto& MeshPair : FoliageMeshes)
+		for (auto& Pair : FoliageInfos)
 		{
 			// Find the per-mesh info matching the mesh.
-			FFoliageMeshInfo& MeshInfo = *MeshPair.Value;
-			UFoliageType* FoliageType = MeshPair.Key;
+			FFoliageInfo& Info = *Pair.Value;
+			UFoliageType* FoliageType = Pair.Key;
 
-			// Make sure the mesh has been PostLoaded as if not it can be considered invalid resulting in a bad HISMC tree
-			UStaticMesh* StaticMesh = FoliageType->GetStaticMesh();
-			if (StaticMesh)
+			// Make sure the source data has been PostLoaded as if not it can be considered invalid resulting in a bad HISMC tree
+			UObject* Source = FoliageType->GetSource();
+			if (Source)
 			{
-				StaticMesh->ConditionalPostLoad();
+				Source->ConditionalPostLoad();
 			}
 
-			if (MeshInfo.Instances.Num() && MeshInfo.Component == nullptr)
+			if (Info.Instances.Num() && !Info.IsInitialized())
 			{
 				FFormatNamedArguments Arguments;
-				if (StaticMesh)
+				if (Source)
 				{
-					Arguments.Add(TEXT("MeshName"), FText::FromString(StaticMesh->GetName()));
+					Arguments.Add(TEXT("MeshName"), FText::FromString(Source->GetName()));
 				}
 				else
 				{
@@ -2819,34 +3458,39 @@ void AInstancedFoliageActor::PostLoad()
 					->AddToken(FTextToken::Create(FText::Format(LOCTEXT("MapCheck_Message_FoliageMissingComponent", "Foliage in this map is missing a component for static mesh {MeshName}. This has been repaired."), Arguments)))
 					->AddToken(FMapErrorToken::Create(FMapErrors::FoliageMissingClusterComponent));
 
-				MeshInfo.ReallocateClusters(this, MeshPair.Key);
+				Info.ReallocateClusters(this, Pair.Key);
 			}
 
 			// Update the hash.
-			MeshInfo.ComponentHash.Empty();
-			MeshInfo.InstanceHash->Empty();
-			for (int32 InstanceIdx = 0; InstanceIdx < MeshInfo.Instances.Num(); InstanceIdx++)
+			Info.ComponentHash.Empty();
+			Info.InstanceHash->Empty();
+			for (int32 InstanceIdx = 0; InstanceIdx < Info.Instances.Num(); InstanceIdx++)
 			{
-				MeshInfo.AddToBaseHash(InstanceIdx);
-				MeshInfo.InstanceHash->InsertInstance(MeshInfo.Instances[InstanceIdx].Location, InstanceIdx);
+				Info.AddToBaseHash(InstanceIdx);
+				Info.InstanceHash->InsertInstance(Info.Instances[InstanceIdx].Location, InstanceIdx);
 			}
 
 			// Convert to Hierarchical foliage
 			if (GetLinkerCustomVersion(FFoliageCustomVersion::GUID) < FFoliageCustomVersion::FoliageUsingHierarchicalISMC)
 			{
-				MeshInfo.ReallocateClusters(this, MeshPair.Key);
+				Info.ReallocateClusters(this, Pair.Key);
 			}
 
 			if (GetLinkerCustomVersion(FFoliageCustomVersion::GUID) < FFoliageCustomVersion::HierarchicalISMCNonTransactional)
 			{
-				if (MeshInfo.Component)
+				check(Info.Type == EFoliageImplType::StaticMesh);
+				if (Info.Type == EFoliageImplType::StaticMesh)
 				{
-					MeshInfo.Component->ClearFlags(RF_Transactional);
+					FFoliageStaticMesh* FoliageStaticMesh = StaticCast<FFoliageStaticMesh*>(Info.Implementation.Get());
+					if (FoliageStaticMesh->Component)
+					{
+						FoliageStaticMesh->Component->ClearFlags(RF_Transactional);
+					}
 				}
 			}
 
 			// Clean up case where embeded instances had their static mesh deleted
-			if (FoliageType->IsNotAssetOrBlueprint() && StaticMesh == nullptr)
+			if (FoliageType->IsNotAssetOrBlueprint() && FoliageType->GetSource() == nullptr)
 			{
 				// We can't remove them here as we are within the loop itself so clean up after
 				FoliageTypeToRemove.Add(FoliageType);
@@ -2857,18 +3501,34 @@ void AInstancedFoliageActor::PostLoad()
 			// Upgrade foliage component
 			if (GetLinkerCustomVersion(FFoliageCustomVersion::GUID) < FFoliageCustomVersion::FoliageUsingFoliageISMC)
 			{
-				MeshInfo.CheckComponentClass(this, FoliageType);
+				check(Info.Type == EFoliageImplType::StaticMesh);
+				if (Info.Type == EFoliageImplType::StaticMesh)
+				{
+					FFoliageStaticMesh* FoliageStaticMesh = StaticCast<FFoliageStaticMesh*>(Info.Implementation.Get());
+					UFoliageType_InstancedStaticMesh* FoliageType_InstancedStaticMesh = Cast<UFoliageType_InstancedStaticMesh>(FoliageType);
+					FoliageStaticMesh->CheckComponentClass(this, FoliageType_InstancedStaticMesh, Info.Instances, Info.SelectedIndices);
+				}
 			}
 
 			// Update foliage component settings if the foliage settings object was changed while the level was not loaded.
-			if (MeshInfo.FoliageTypeUpdateGuid != FoliageType->UpdateGuid)
+			if (Info.FoliageTypeUpdateGuid != FoliageType->UpdateGuid)
 			{
-				if (MeshInfo.FoliageTypeUpdateGuid.IsValid())
+				if (Info.FoliageTypeUpdateGuid.IsValid())
 				{
-					MeshInfo.CheckComponentClass(this, FoliageType);
-					MeshInfo.UpdateComponentSettings(FoliageType);
+					if (Info.Type == EFoliageImplType::StaticMesh)
+					{
+						FFoliageStaticMesh* FoliageStaticMesh = StaticCast<FFoliageStaticMesh*>(Info.Implementation.Get());
+						UFoliageType_InstancedStaticMesh* FoliageType_InstancedStaticMesh = Cast<UFoliageType_InstancedStaticMesh>(FoliageType);
+						FoliageStaticMesh->CheckComponentClass(this, FoliageType_InstancedStaticMesh, Info.Instances, Info.SelectedIndices);
+						FoliageStaticMesh->UpdateComponentSettings(FoliageType_InstancedStaticMesh);
+					}
+					else if (Info.Type == EFoliageImplType::Actor)
+					{
+						FFoliageActor* FoliageActor = StaticCast<FFoliageActor*>(Info.Implementation.Get());
+						FoliageActor->Reapply(this, FoliageType, Info.Instances);
+					}
 				}
-				MeshInfo.FoliageTypeUpdateGuid = FoliageType->UpdateGuid;
+				Info.FoliageTypeUpdateGuid = FoliageType->UpdateGuid;
 			}
 		}
 
@@ -2945,15 +3605,25 @@ void AInstancedFoliageActor::PostLoad()
 
 	if (!GIsEditor && CVarFoliageDiscardDataOnLoad.GetValueOnGameThread())
 	{
-		for (auto& MeshPair : FoliageMeshes)
+		for (auto& Pair : FoliageInfos)
 		{
-			if (MeshPair.Value->Component != nullptr && (!MeshPair.Key || MeshPair.Key->bEnableDensityScaling))
+			if (Pair.Value->Type == EFoliageImplType::StaticMesh)
 			{
-				MeshPair.Value->Component->ConditionalPostLoad();
-				MeshPair.Value->Component->DestroyComponent();
+				FFoliageStaticMesh* FoliageStaticMesh = StaticCast<FFoliageStaticMesh*>(Pair.Value->Implementation.Get());
+				
+				if (FoliageStaticMesh->Component != nullptr && (!Pair.Key || Pair.Key->bEnableDensityScaling))
+				{
+					FoliageStaticMesh->Component->ConditionalPostLoad();
+					FoliageStaticMesh->Component->DestroyComponent();
+				}
+			}	
+			else if (Pair.Value->Type == EFoliageImplType::Actor)
+			{
+				FFoliageActor* FoliageActor = StaticCast<FFoliageActor*>(Pair.Value->Implementation.Get());
+				FoliageActor->DestroyActors(true);
 			}
 
-			MeshPair.Value = FFoliageMeshInfo();
+			Pair.Value = FFoliageInfo();
 		}
 	}
 }
@@ -2962,13 +3632,13 @@ void AInstancedFoliageActor::PostLoad()
 
 void AInstancedFoliageActor::RepairDuplicateIFA(AInstancedFoliageActor* DuplicateIFA)
 {
-	for (auto& MeshPair : DuplicateIFA->FoliageMeshes)
+	for (auto& Pair : DuplicateIFA->FoliageInfos)
 	{
-		UFoliageType* DupeFoliageType = MeshPair.Key;
-		FFoliageMeshInfo& DupeMeshInfo = *MeshPair.Value;
+		UFoliageType* DupeFoliageType = Pair.Key;
+		FFoliageInfo& DupeMeshInfo = *Pair.Value;
 
 		// Get foliage type compatible with target IFA
-		FFoliageMeshInfo* TargetMeshInfo = nullptr;
+		FFoliageInfo* TargetMeshInfo = nullptr;
 		UFoliageType* TargetFoliageType = AddFoliageType(DupeFoliageType, &TargetMeshInfo);
 
 		// Copy the instances
@@ -2976,40 +3646,30 @@ void AInstancedFoliageActor::RepairDuplicateIFA(AInstancedFoliageActor* Duplicat
 		{
 			if ((Instance.Flags & FOLIAGE_InstanceDeleted) == 0)
 			{
-				TargetMeshInfo->AddInstance(this, TargetFoliageType, Instance, false);
+				TargetMeshInfo->AddInstance(this, TargetFoliageType, Instance);
 			}
 		}
 
-		TargetMeshInfo->Component->BuildTreeIfOutdated(true, true);
+		TargetMeshInfo->Refresh(this, true, true);
 	}
 
 	GetWorld()->DestroyActor(DuplicateIFA);
 }
 
-void AInstancedFoliageActor::NotifyFoliageTypeChanged(UFoliageType* FoliageType, bool bMeshChanged)
+void AInstancedFoliageActor::NotifyFoliageTypeChanged(UFoliageType* FoliageType, bool bSourceChanged)
 {
-	FFoliageMeshInfo* TypeInfo = FindMesh(FoliageType);
+	FFoliageInfo* TypeInfo = FindInfo(FoliageType);
 
 	if (TypeInfo)
 	{
-		TypeInfo->CheckComponentClass(this, FoliageType);
-		TypeInfo->UpdateComponentSettings(FoliageType);
+		TypeInfo->NotifyFoliageTypeChanged(this, FoliageType, bSourceChanged);	
 
-		if (bMeshChanged)
+		if (bSourceChanged)
 		{
 			// If the type's mesh has changed, the UI needs to be notified so it can update thumbnails accordingly
-			OnFoliageTypeMeshChangedEvent.Broadcast(FoliageType);
-
-			// Change bounds delegate bindings
-			if (TypeInfo->Component != nullptr && TypeInfo->Component->GetStaticMesh() != nullptr && FoliageType->GetStaticMesh() != nullptr)
-			{
-				TypeInfo->Component->GetStaticMesh()->GetOnExtendedBoundsChanged().AddRaw(TypeInfo, &FFoliageMeshInfo::HandleComponentMeshBoundsChanged);
-
-				// Mesh changed, so we must update the occlusion tree
-				TypeInfo->Component->BuildTreeIfOutdated(true, false);
-			}
+			OnFoliageTypeMeshChangedEvent.Broadcast(FoliageType);		
 
-			if (FoliageType->IsNotAssetOrBlueprint() && FoliageType->GetStaticMesh() == nullptr) //If the mesh has been deleted and we're a per foliage actor instance we must remove all instances of the mesh
+			if (FoliageType->IsNotAssetOrBlueprint() && FoliageType->GetSource() == nullptr) //If the source data has been deleted and we're a per foliage actor instance we must remove all instances 
 			{
 				RemoveFoliageType(&FoliageType, 1);
 			}
@@ -3017,20 +3677,14 @@ void AInstancedFoliageActor::NotifyFoliageTypeChanged(UFoliageType* FoliageType,
 	}
 }
 
-void AInstancedFoliageActor::NotifyFoliageTypeWillChange(UFoliageType* FoliageType, bool bMeshChanged)
+void AInstancedFoliageActor::NotifyFoliageTypeWillChange(UFoliageType* FoliageType)
 {
-	if (bMeshChanged)
-	{
-		FFoliageMeshInfo* TypeInfo = FindMesh(FoliageType);
+	FFoliageInfo* TypeInfo = FindInfo(FoliageType);
 
-		// Change bounds delegate bindings
-		if (TypeInfo)
-		{
-			if (TypeInfo->Component != nullptr && TypeInfo->Component->GetStaticMesh() != nullptr)
-			{
-				TypeInfo->Component->GetStaticMesh()->GetOnExtendedBoundsChanged().RemoveAll(TypeInfo);
-			}
-		}
+	// Change bounds delegate bindings
+	if (TypeInfo)
+	{
+		TypeInfo->NotifyFoliageTypeWillChange(this, FoliageType);
 	}
 }
 
@@ -3076,19 +3730,19 @@ void AInstancedFoliageActor::OnPostApplyLevelOffset(ULevel* InLevel, UWorld* InW
 
 	if (GIsEditor && InWorld && !InWorld->IsGameWorld())
 	{
-		for (auto& MeshPair : FoliageMeshes)
+		for (auto& Pair : FoliageInfos)
 		{
-			FFoliageMeshInfo& MeshInfo = *MeshPair.Value;
+			FFoliageInfo& Info = *Pair.Value;
 
 			InstanceBaseCache.UpdateInstanceBaseCachedTransforms();
 
-			MeshInfo.InstanceHash->Empty();
-			for (int32 InstanceIdx = 0; InstanceIdx < MeshInfo.Instances.Num(); InstanceIdx++)
+			Info.InstanceHash->Empty();
+			for (int32 InstanceIdx = 0; InstanceIdx < Info.Instances.Num(); InstanceIdx++)
 			{
-				FFoliageInstance& Instance = MeshInfo.Instances[InstanceIdx];
+				FFoliageInstance& Instance = Info.Instances[InstanceIdx];
 				Instance.Location += InOffset;
 				// Rehash instance location
-				MeshInfo.InstanceHash->InsertInstance(Instance.Location, InstanceIdx);
+				Info.InstanceHash->InsertInstance(Instance.Location, InstanceIdx);
 			}
 		}
 	}
@@ -3097,26 +3751,26 @@ void AInstancedFoliageActor::OnPostApplyLevelOffset(ULevel* InLevel, UWorld* InW
 
 void AInstancedFoliageActor::CleanupDeletedFoliageType()
 {
-	for (auto& MeshPair : FoliageMeshes)
+	for (auto& Pair : FoliageInfos)
 	{
-		if (MeshPair.Key == nullptr)
+		if (Pair.Key == nullptr)
 		{
-			FFoliageMeshInfo& MeshInfo = *MeshPair.Value;
+			FFoliageInfo& Info = *Pair.Value;
 			TArray<int32> InstancesToRemove;
-			for (int32 InstanceIdx = 0; InstanceIdx < MeshInfo.Instances.Num(); InstanceIdx++)
+			for (int32 InstanceIdx = 0; InstanceIdx < Info.Instances.Num(); InstanceIdx++)
 			{
 				InstancesToRemove.Add(InstanceIdx);
 			}
 
 			if (InstancesToRemove.Num())
 			{
-				MeshInfo.RemoveInstances(this, InstancesToRemove, true);
+				Info.RemoveInstances(this, InstancesToRemove, true);
 			}
 		}
 
 	}
 
-	while (FoliageMeshes.Remove(nullptr)) {}	//remove entries from the map
+	while (FoliageInfos.Remove(nullptr)) {}	//remove entries from the map
 
 }
 
@@ -3129,15 +3783,12 @@ void AInstancedFoliageActor::AddReferencedObjects(UObject* InThis, FReferenceCol
 {
 	AInstancedFoliageActor* This = CastChecked<AInstancedFoliageActor>(InThis);
 
-	for (auto& MeshPair : This->FoliageMeshes)
+	for (auto& Pair : This->FoliageInfos)
 	{
-		Collector.AddReferencedObject(MeshPair.Key, This);
-		FFoliageMeshInfo& MeshInfo = *MeshPair.Value;
+		Collector.AddReferencedObject(Pair.Key, This);
+		FFoliageInfo& Info = *Pair.Value;
 
-		if (MeshInfo.Component)
-		{
-			Collector.AddReferencedObject(MeshInfo.Component, This);
-		}
+		Info.AddReferencedObjects(This, Collector);
 	}
 
 	Super::AddReferencedObjects(This, Collector);
@@ -3165,7 +3816,7 @@ bool AInstancedFoliageActor::FoliageTrace(const UWorld* InWorld, FHitResult& Out
 	for (const FHitResult& Hit : Hits)
 	{
 		const AActor* HitActor = Hit.GetActor();
-
+				
 		// don't place procedural foliage inside an AProceduralFoliageBlockingVolume
 		// this test is first because two of the tests below would otherwise cause the trace to ignore AProceduralFoliageBlockingVolume
 		if (DesiredInstance.PlacementMode == EFoliagePlacementMode::Procedural)
@@ -3200,11 +3851,18 @@ bool AInstancedFoliageActor::FoliageTrace(const UWorld* InWorld, FHitResult& Out
 		}
 
 		// Don't place foliage on itself
-		if (const AInstancedFoliageActor* FoliageActor = Cast<AInstancedFoliageActor>(HitActor))
+		const AInstancedFoliageActor* FoliageActor = Cast<AInstancedFoliageActor>(HitActor);
+		if (!FoliageActor && HitActor && AInstancedFoliageActor::IsOwnedByFoliage(HitActor))
 		{
-			if (const FFoliageMeshInfo* FoundMeshInfo = FoliageActor->FindMesh(DesiredInstance.FoliageType))
+			FoliageActor = HitActor->GetLevel()->InstancedFoliageActor.Get();
+			if (FoliageActor == nullptr)
+			{
+				continue;
+			}
+
+			if (const FFoliageInfo* FoundMeshInfo = FoliageActor->FindInfo(DesiredInstance.FoliageType))
 			{
-				if (FoundMeshInfo->Component == HitComponent)
+				if (FoundMeshInfo->Implementation->IsOwnedComponent(HitComponent))
 				{
 					continue;
 				}
@@ -3225,18 +3883,28 @@ bool AInstancedFoliageActor::FoliageTrace(const UWorld* InWorld, FHitResult& Out
 		}
 
 		OutHit = Hit;
-
+			
 		// When placing foliage on other foliage, we need to return the base component of the other foliage, not the foliage component, so that it moves correctly
-		if (const AInstancedFoliageActor* FoliageActor = Cast<AInstancedFoliageActor>(HitActor))
+		if (FoliageActor)
 		{
-			for (auto& MeshPair : FoliageActor->FoliageMeshes)
+			for (auto& Pair : FoliageActor->FoliageInfos)
 			{
-				const FFoliageMeshInfo& MeshInfo = *MeshPair.Value;
-				if (MeshInfo.Component == HitComponent)
+				const FFoliageInfo& Info = *Pair.Value;
+
+				if (Hit.Item != INDEX_NONE && Info.Implementation->IsOwnedComponent(HitComponent))
 				{
-					OutHit.Component = CastChecked<UPrimitiveComponent>(FoliageActor->InstanceBaseCache.GetInstanceBasePtr(MeshInfo.Instances[Hit.Item].BaseId).Get(), ECastCheckedType::NullAllowed);
+					OutHit.Component = CastChecked<UPrimitiveComponent>(FoliageActor->InstanceBaseCache.GetInstanceBasePtr(Info.Instances[Hit.Item].BaseId).Get(), ECastCheckedType::NullAllowed);
 					break;
 				}
+				else
+				{
+					int32 InstanceIndex = Info.Implementation->FindIndex(HitComponent);
+					if (InstanceIndex != INDEX_NONE)
+					{
+						OutHit.Component = CastChecked<UPrimitiveComponent>(FoliageActor->InstanceBaseCache.GetInstanceBasePtr(Info.Instances[InstanceIndex].BaseId).Get(), ECastCheckedType::NullAllowed);
+						break;
+					}
+				}				
 			}
 		}
 
diff --git a/Engine/Source/Runtime/Foliage/Private/ProceduralFoliageInstance.cpp b/Engine/Source/Runtime/Foliage/Private/ProceduralFoliageInstance.cpp
--- a/Engine/Source/Runtime/Foliage/Private/ProceduralFoliageInstance.cpp
+++ b/Engine/Source/Runtime/Foliage/Private/ProceduralFoliageInstance.cpp
@@ -57,8 +57,8 @@ FProceduralFoliageInstance* GetLessFit(FProceduralFoliageInstance* A, FProcedura
 
 FProceduralFoliageInstance* FProceduralFoliageInstance::Domination(FProceduralFoliageInstance* A, FProceduralFoliageInstance* B, ESimulationOverlap::Type OverlapType)
 {
-	const UFoliageType_InstancedStaticMesh* AType = A->Type;
-	const UFoliageType_InstancedStaticMesh* BType = B->Type;
+	const UFoliageType* AType = A->Type;
+	const UFoliageType* BType = B->Type;
 
 	FProceduralFoliageInstance* Dominated = GetLessFit(A, B);
 
diff --git a/Engine/Source/Runtime/Foliage/Private/ProceduralFoliageTile.cpp b/Engine/Source/Runtime/Foliage/Private/ProceduralFoliageTile.cpp
--- a/Engine/Source/Runtime/Foliage/Private/ProceduralFoliageTile.cpp
+++ b/Engine/Source/Runtime/Foliage/Private/ProceduralFoliageTile.cpp
@@ -57,7 +57,7 @@ bool UProceduralFoliageTile::HandleOverlaps(FProceduralFoliageInstance* Instance
 	return bSurvived;
 }
 
-FProceduralFoliageInstance* UProceduralFoliageTile::NewSeed(const FVector& Location, float Scale, const UFoliageType_InstancedStaticMesh* Type, float InAge, bool bBlocker)
+FProceduralFoliageInstance* UProceduralFoliageTile::NewSeed(const FVector& Location, float Scale, const UFoliageType* Type, float InAge, bool bBlocker)
 {
 	const float InitRadius = Type->GetMaxRadius() * Scale;
 	{
@@ -98,7 +98,7 @@ FProceduralFoliageInstance* UProceduralFoliageTile::NewSeed(const FVector& Locat
 
 float GetSeedMinDistance(const FProceduralFoliageInstance* Instance, const float NewInstanceAge, const int32 SimulationStep)
 {
-	const UFoliageType_InstancedStaticMesh* Type = Instance->Type;
+	const UFoliageType* Type = Instance->Type;
 	const int32 StepsLeft = Type->MaxAge - SimulationStep;
 	const float InstanceMaxAge = Type->GetNextAge(Instance->Age, StepsLeft);
 	const float NewInstanceMaxAge = Type->GetNextAge(NewInstanceAge, StepsLeft);
@@ -123,7 +123,7 @@ float UProceduralFoliageTile::GetRandomGaussian()
 	return Z1;
 }
 
-FVector UProceduralFoliageTile::GetSeedOffset(const UFoliageType_InstancedStaticMesh* Type, float MinDistance)
+FVector UProceduralFoliageTile::GetSeedOffset(const UFoliageType* Type, float MinDistance)
 {
 	//We want 10% of seeds to be the max distance so we use a z score of +- 1.64
 	const float MaxZScore = 1.64f;
@@ -145,7 +145,7 @@ void UProceduralFoliageTile::AgeSeeds()
 		if (UserCancelled()){ return; }
 		if (Instance->IsAlive())
 		{
-			const UFoliageType_InstancedStaticMesh* Type = Instance->Type;
+			const UFoliageType* Type = Instance->Type;
 			if (SimulationStep <= Type->NumSteps && Type->GetSpawnsInShade() == bSimulateOnlyInShade)
 			{
 				const float CurrentAge = Instance->Age;
@@ -186,7 +186,7 @@ void UProceduralFoliageTile::SpreadSeeds(TArray<FProceduralFoliageInstance*>& Ne
 			continue;
 		}
 
-		const UFoliageType_InstancedStaticMesh* Type = Inst->Type;
+		const UFoliageType* Type = Inst->Type;
 
 		if (SimulationStep <= Type->NumSteps  && Type->GetSpawnsInShade() == bSimulateOnlyInShade)
 		{
@@ -220,12 +220,12 @@ void UProceduralFoliageTile::AddRandomSeeds(TArray<FProceduralFoliageInstance*>&
 	TMap<const UFoliageType*, int32> SeedsLeftMap;
 	TMap<const UFoliageType*, FRandomStream> RandomStreamPerType;
 
-	TArray<const UFoliageType_InstancedStaticMesh*> TypesToSeed;
+	TArray<const UFoliageType*> TypesToSeed;
 
 	for (const FFoliageTypeObject& FoliageTypeObject : FoliageSpawner->GetFoliageTypes())
 	{
 		if (UserCancelled()){ return; }
-		const UFoliageType_InstancedStaticMesh* TypeInstance = FoliageTypeObject.GetInstance();
+		const UFoliageType* TypeInstance = FoliageTypeObject.GetInstance();
 		if (TypeInstance && TypeInstance->GetSpawnsInShade() == bSimulateOnlyInShade)
 		{
 			{	//compute the number of initial seeds
@@ -277,7 +277,7 @@ void UProceduralFoliageTile::AddRandomSeeds(TArray<FProceduralFoliageInstance*>&
 		if (UserCancelled()){ return; }
 		TypeIdx = (TypeIdx + 1) % NumTypes;	//keep cycling through the types that we spawn initial seeds for to make sure everyone gets fair chance
 
-		if (const UFoliageType_InstancedStaticMesh* Type = TypesToSeed[TypeIdx])
+		if (const UFoliageType* Type = TypesToSeed[TypeIdx])
 		{
 			int32& SeedsLeft = SeedsLeftMap.FindChecked(Type);
 			if (SeedsLeft == 0)
@@ -426,7 +426,7 @@ void UProceduralFoliageTile::RunSimulation(const int32 MaxNumSteps, bool bOnlyIn
 
 	for (const FFoliageTypeObject& FoliageTypeObject : FoliageSpawner->GetFoliageTypes())
 	{
-		const UFoliageType_InstancedStaticMesh* TypeInstance = FoliageTypeObject.GetInstance();
+		const UFoliageType* TypeInstance = FoliageTypeObject.GetInstance();
 		if (TypeInstance && TypeInstance->GetSpawnsInShade() == bOnlyInShade)
 		{
 			MaxSteps = FMath::Max(MaxSteps, TypeInstance->NumSteps + 1);
diff --git a/Engine/Source/Runtime/Foliage/Public/FoliageType.h b/Engine/Source/Runtime/Foliage/Public/FoliageType.h
--- a/Engine/Source/Runtime/Foliage/Public/FoliageType.h
+++ b/Engine/Source/Runtime/Foliage/Public/FoliageType.h
@@ -84,25 +84,24 @@ class UFoliageType : public UObject
 {
 	GENERATED_UCLASS_BODY()
 
-	/* Gets/Sets the mesh associated with this FoliageType */
-	virtual UStaticMesh* GetStaticMesh() const PURE_VIRTUAL(UFoliageType::GetStaticMesh, return nullptr; );
-	virtual void SetStaticMesh(UStaticMesh* InStaticMesh) PURE_VIRTUAL(UFoliageType::SetStaticMesh,);
-
-	/* Gets the component class to use for instances of this FoliageType */
-	virtual UClass* GetComponentClass() const PURE_VIRTUAL(UFoliageType::GetComponentClass(), return nullptr;);
-	
+	/* Gets/Sets the source data associated with this FoliageType */
+	virtual UObject* GetSource() const PURE_VIRTUAL(UFoliageType::GetSource, return nullptr; );
+		
 	virtual void Serialize(FArchive& Ar) override;
 
 	virtual bool IsNotAssetOrBlueprint() const;
 
 	FOLIAGE_API FVector GetRandomScale() const;
-
+	
 #if WITH_EDITOR
+	virtual void SetSource(UObject* InSource) PURE_VIRTUAL(UFoliageType::SetSource, );
+	virtual void UpdateBounds() {}
 	/* Lets subclasses decide if the InstancedFoliageActor should reallocate its instances if the specified property change event occurs */
-	virtual bool IsFoliageReallocationRequiredForPropertyChange(struct FPropertyChangedEvent& PropertyChangedEvent) const { return true; }
+	virtual bool IsFoliageReallocationRequiredForPropertyChange(const UProperty* Property) const { return true; }
 
 	virtual void PreEditChange(UProperty* PropertyAboutToChange) override;
 	virtual void PostEditChangeProperty(struct FPropertyChangedEvent& PropertyChangedEvent) override;
+	virtual bool IsSourcePropertyChange(const UProperty* Property) const { return false; }
 
 	/* Notifies all relevant foliage actors that HiddenEditorView mask has been changed */
 	FOLIAGE_API void OnHiddenEditorViewMaskChanged(UWorld* InWorld);
@@ -129,7 +128,7 @@ public:
 	float DensityAdjustmentFactor;
 
 	/** The minimum distance between foliage instances */
-	UPROPERTY(EditAnywhere, Category=Painting, meta=(UIMin=0, ClampMin=0, UIMax = 1000, ClampMax = 1000, ReapplyCondition="ReapplyRadius"))
+	UPROPERTY(EditAnywhere, Category=Painting, meta=(UIMin=0, ClampMin=0, UIMax = 100000, ClampMax = 100000, ReapplyCondition="ReapplyRadius"))
 	float Radius;
 
 	/** Specifies foliage instance scaling behavior when painting. */
diff --git a/Engine/Source/Runtime/Foliage/Public/FoliageTypeObject.h b/Engine/Source/Runtime/Foliage/Public/FoliageTypeObject.h
--- a/Engine/Source/Runtime/Foliage/Public/FoliageTypeObject.h
+++ b/Engine/Source/Runtime/Foliage/Public/FoliageTypeObject.h
@@ -22,10 +22,10 @@ struct FOLIAGE_API FFoliageTypeObject
 	void RefreshInstance();
 
 	/** Gets the instance of this foliage type. Creates the instance if needed. */
-	const UFoliageType_InstancedStaticMesh* GetInstance();
+	const UFoliageType* GetInstance();
 
 	/** Gets the instance of this foliage type. */
-	const UFoliageType_InstancedStaticMesh* GetInstance() const;
+	const UFoliageType* GetInstance() const;
 
 	/** @return Whether this would return a valid instance */
 	bool ContainsValidInstance() const;
@@ -40,12 +40,12 @@ struct FOLIAGE_API FFoliageTypeObject
 
 private:
 	/** The foliage type that will be spawned by the procedural foliage simulation */
-	UPROPERTY(Category = ProceduralFoliageSimulation, EditAnywhere, meta=(AllowedClasses="FoliageType_InstancedStaticMesh,Blueprint", DisplayThumbnail="true", ThumbnailSize="X=40 Y=40"))
+	UPROPERTY(Category = ProceduralFoliageSimulation, EditAnywhere, meta=(AllowedClasses="FoliageType_InstancedStaticMesh,FoliageType_Actor,Blueprint", DisplayThumbnail="true", ThumbnailSize="X=40 Y=40"))
 	UObject* FoliageTypeObject;
 
 	/** The actual instance of the foliage type that is used for spawning */
 	UPROPERTY(transient)
-	UFoliageType_InstancedStaticMesh* TypeInstance;
+	UFoliageType* TypeInstance;
 
 	/** Whether this contains an asset object (as opposed to a BP class) */
 	UPROPERTY()
diff --git a/Engine/Source/Runtime/Foliage/Public/FoliageType_Actor.h b/Engine/Source/Runtime/Foliage/Public/FoliageType_Actor.h
new file mode 100644
--- /dev/null
+++ b/Engine/Source/Runtime/Foliage/Public/FoliageType_Actor.h
@@ -0,0 +1,30 @@
+// Copyright 1998-2019 Epic Games, Inc. All Rights Reserved.
+#pragma once
+
+#include "Templates/SubclassOf.h"
+#include "FoliageType.h"
+#include "FoliageType_Actor.generated.h"
+
+UCLASS(hidecategories = Object, editinlinenew, MinimalAPI)
+class UFoliageType_Actor : public UFoliageType
+{
+	GENERATED_UCLASS_BODY()
+	
+	UPROPERTY(EditAnywhere, Category = Actor)
+	TSubclassOf<AActor> ActorClass;
+			
+	virtual UObject* GetSource() const override { return ActorClass; }
+
+#if WITH_EDITOR
+	virtual void UpdateBounds();
+	virtual bool IsSourcePropertyChange(const UProperty* Property) const override
+	{
+		return Property && Property->GetFName() == GET_MEMBER_NAME_CHECKED(UFoliageType_Actor, ActorClass);
+	}
+	virtual void SetSource(UObject* InSource) override
+	{
+		ActorClass = Cast<UClass>(InSource);
+		UpdateBounds();
+	}
+#endif
+};
\ No newline at end of file
diff --git a/Engine/Source/Runtime/Foliage/Public/FoliageType_InstancedStaticMesh.h b/Engine/Source/Runtime/Foliage/Public/FoliageType_InstancedStaticMesh.h
--- a/Engine/Source/Runtime/Foliage/Public/FoliageType_InstancedStaticMesh.h
+++ b/Engine/Source/Runtime/Foliage/Public/FoliageType_InstancedStaticMesh.h
@@ -24,19 +24,31 @@ class UFoliageType_InstancedStaticMesh : public UFoliageType
 	  * You can make a Blueprint subclass of FoliageInstancedStaticMeshComponent to implement custom behavior and assign that class here. */
 	UPROPERTY(EditAnywhere, Category = Mesh)
 	TSubclassOf<UFoliageInstancedStaticMeshComponent> ComponentClass;
-
-	virtual UStaticMesh* GetStaticMesh() const override
+		
+	UStaticMesh* GetStaticMesh() const
 	{
 		return Mesh;
 	}
 
-	virtual void SetStaticMesh(UStaticMesh* InStaticMesh) override
+	UClass* GetComponentClass() const
 	{
-		Mesh = InStaticMesh;
+		return *ComponentClass;
+	}
+
+	virtual UObject* GetSource() const override;
+
+#if WITH_EDITOR
+	virtual void UpdateBounds() override;
+	virtual bool IsSourcePropertyChange(const UProperty* Property) const override
+	{
+		return Property && Property->GetFName() == GET_MEMBER_NAME_CHECKED(UFoliageType_InstancedStaticMesh, Mesh);
 	}
+	virtual void SetSource(UObject* InSource) override;
 
-	virtual UClass* GetComponentClass() const override
+	void SetStaticMesh(UStaticMesh* InStaticMesh)
 	{
-		return *ComponentClass;
+		Mesh = InStaticMesh;
+		UpdateBounds();
 	}
+#endif
 };
diff --git a/Engine/Source/Runtime/Foliage/Public/InstancedFoliage.h b/Engine/Source/Runtime/Foliage/Public/InstancedFoliage.h
--- a/Engine/Source/Runtime/Foliage/Public/InstancedFoliage.h
+++ b/Engine/Source/Runtime/Foliage/Public/InstancedFoliage.h
@@ -13,15 +13,9 @@ class AInstancedFoliageActor;
 class UActorComponent;
 class UFoliageType;
 class UHierarchicalInstancedStaticMeshComponent;
+class UFoliageType_InstancedStaticMesh;
 class UPrimitiveComponent;
-struct FFoliageInstanceHash;
-
-//
-// Forward declarations.
-//
-class UHierarchicalInstancedStaticMeshComponent;
-class AInstancedFoliageActor;
-class UFoliageType;
+class UStaticMesh;
 struct FFoliageInstanceHash;
 
 DECLARE_LOG_CATEGORY_EXTERN(LogInstancedFoliage, Log, All);
@@ -143,12 +137,91 @@ struct FFoliageMeshInfo_Deprecated
 };
 
 /**
-*	FFoliageMeshInfo - editor info for all matching foliage meshes
+*	FFoliageInfo - editor info for all matching foliage meshes
 */
-struct FFoliageMeshInfo
+struct FFoliageMeshInfo_Deprecated2
 {
 	UHierarchicalInstancedStaticMeshComponent* Component;
 
+#if WITH_EDITORONLY_DATA
+	// Allows us to detect if FoliageType was updated while this level wasn't loaded
+	FGuid FoliageTypeUpdateGuid;
+
+	// Editor-only placed instances
+	TArray<FFoliageInstance> Instances;
+#endif
+
+
+	FFoliageMeshInfo_Deprecated2();
+	
+	friend FArchive& operator<<(FArchive& Ar, FFoliageMeshInfo_Deprecated2& MeshInfo);
+};
+
+///
+/// EFoliageImplType
+///
+enum class EFoliageImplType : uint8
+{
+	Unknown = 0,
+	StaticMesh = 1,
+	Actor = 2
+};
+
+///
+/// FFoliageInfoImpl
+///
+struct FFoliageImpl
+{
+	FFoliageImpl() {}
+	virtual ~FFoliageImpl() {}
+
+	virtual void AddReferencedObjects(UObject* InThis, FReferenceCollector& Collector) {}
+	virtual void Serialize(FArchive& Ar) = 0;
+
+#if WITH_EDITOR
+	virtual bool IsInitialized() const = 0;
+	virtual void Initialize(AInstancedFoliageActor* IFA, const UFoliageType* FoliageType) = 0;
+	virtual void Uninitialize() = 0;
+	virtual int32 GetInstanceCount() const = 0;
+	virtual void PreAddInstances(AInstancedFoliageActor* IFA, const UFoliageType* FoliageType, int32 Count) = 0;
+	virtual void AddInstance(AInstancedFoliageActor* IFA, const FFoliageInstance& NewInstance) = 0;
+	virtual void RemoveInstance(int32 InstanceIndex) = 0;
+	virtual void SetInstanceWorldTransform(int32 InstanceIndex, const FTransform& Transform, bool bTeleport) = 0;
+	virtual FTransform GetInstanceWorldTransform(int32 InstanceIndex) const = 0;
+	virtual void PostUpdateInstances() {}
+	virtual bool IsOwnedComponent(const UPrimitiveComponent* Component) const = 0;
+	virtual int32 FindIndex(const UPrimitiveComponent* HitComponent) const { return INDEX_NONE; }
+
+	virtual void SelectInstances(bool bSelect, int32 InstanceIndex, int32 Count) = 0;
+	virtual void ApplySelection(bool bApply, const TSet<int32>& SelectedIndices) = 0;
+	virtual void ClearSelection(const TSet<int32>& SelectedIndices) = 0;
+
+	virtual void BeginUpdate() {}
+	virtual void EndUpdate() {}
+	virtual void Refresh(AInstancedFoliageActor* IFA, const TArray<FFoliageInstance>& Instances, bool Async, bool Force) {}
+	virtual void OnHiddenEditorViewMaskChanged(uint64 InHiddenEditorViews) = 0;
+	virtual void PreEditUndo(AInstancedFoliageActor* IFA, UFoliageType* FoliageType) {}
+	virtual void PostEditUndo(AInstancedFoliageActor* IFA, UFoliageType* FoliageType, const TArray<FFoliageInstance>& Instances, const TSet<int32>& SelectedIndices) = 0;
+	virtual void NotifyFoliageTypeWillChange(AInstancedFoliageActor* IFA, UFoliageType* FoliageType) {}
+	virtual void NotifyFoliageTypeChanged(AInstancedFoliageActor* IFA, UFoliageType* FoliageType, const TArray<FFoliageInstance>& Instances, const TSet<int32>& SelectedIndices, bool bSourceChanged) = 0;
+	virtual void EnterEditMode() {}
+	virtual void ExitEditMode() {}
+#endif
+
+	virtual int32 GetOverlappingSphereCount(const FSphere& Sphere) const { return 0; }
+	virtual int32 GetOverlappingBoxCount(const FBox& Box) const { return 0; }
+	virtual void GetOverlappingBoxTransforms(const FBox& Box, TArray<FTransform>& OutTransforms) const { }
+	virtual void GetOverlappingMeshCount(const FSphere& Sphere, TMap<UStaticMesh*, int32>& OutCounts) const { }
+};
+
+///
+/// FFoliageInfo
+///
+struct FFoliageInfo
+{
+	EFoliageImplType Type;
+	TUniquePtr<FFoliageImpl> Implementation;
+
 #if WITH_EDITORONLY_DATA
 	// Allows us to detect if FoliageType was updated while this level wasn't loaded
 	FGuid FoliageTypeUpdateGuid;
@@ -166,46 +239,44 @@ struct FFoliageMeshInfo
 	TSet<int32> SelectedIndices;
 #endif
 
-	FOLIAGE_API FFoliageMeshInfo();
-	FOLIAGE_API ~FFoliageMeshInfo();
-
-	FFoliageMeshInfo(FFoliageMeshInfo&& Other)
-		// even VC++2013 doesn't support "=default" on move constructors
-		: Component(Other.Component)
-#if WITH_EDITORONLY_DATA
-		, FoliageTypeUpdateGuid(MoveTemp(Other.FoliageTypeUpdateGuid))
-		, Instances(MoveTemp(Other.Instances))
-		, InstanceHash(MoveTemp(Other.InstanceHash))
-		, ComponentHash(MoveTemp(Other.ComponentHash))
-		, SelectedIndices(MoveTemp(Other.SelectedIndices))
-#endif
-	{ }
+	FOLIAGE_API FFoliageInfo();
+	FOLIAGE_API ~FFoliageInfo();
+
+	FFoliageInfo(FFoliageInfo&& Other) = default;
+	FFoliageInfo& operator=(FFoliageInfo&& Other) = default;
+
+	// Will only return a valid component in the case of non-actor foliage
+	FOLIAGE_API UHierarchicalInstancedStaticMeshComponent* GetComponent() const;
+
+	void AddReferencedObjects(UObject* InThis, FReferenceCollector& Collector);
+		
+	FOLIAGE_API void CreateImplementation(EFoliageImplType InType);
+	FOLIAGE_API void Initialize(AInstancedFoliageActor* InIFA, const UFoliageType* FoliageType);
+	FOLIAGE_API void Uninitialize();
+	FOLIAGE_API bool IsInitialized() const;
+		
+	int32 GetOverlappingSphereCount(const FSphere& Sphere) const;
+	int32 GetOverlappingBoxCount(const FBox& Box) const;
+	void GetOverlappingBoxTransforms(const FBox& Box, TArray<FTransform>& OutTransforms) const;
+	void GetOverlappingMeshCount(const FSphere& Sphere, TMap<UStaticMesh*, int32>& OutCounts) const;
+	
+#if WITH_EDITOR
+	
+	FOLIAGE_API void CreateImplementation(const UFoliageType* FoliageType);
+	FOLIAGE_API void NotifyFoliageTypeWillChange(AInstancedFoliageActor* IFA, UFoliageType* FoliageType);
+	FOLIAGE_API void NotifyFoliageTypeChanged(AInstancedFoliageActor* IFA, UFoliageType* FoliageType, bool bSourceChanged);
+	
+	FOLIAGE_API void ClearSelection();
 
-	FFoliageMeshInfo& operator=(FFoliageMeshInfo&& Other)
-		// even VC++2013 doesn't support "=default" on move assignment
-	{
-		Component = Other.Component;
-#if WITH_EDITORONLY_DATA
-		FoliageTypeUpdateGuid = MoveTemp(Other.FoliageTypeUpdateGuid);
-		Instances = MoveTemp(Other.Instances);
-		InstanceHash = MoveTemp(Other.InstanceHash);
-		ComponentHash = MoveTemp(Other.ComponentHash);
-		SelectedIndices = MoveTemp(Other.SelectedIndices);
-#endif
+	FOLIAGE_API void SetInstanceWorldTransform(int32 InstanceIndex, const FTransform& Transform, bool bTeleport);
 
-		return *this;
-	}
+	FOLIAGE_API void SetRandomSeed(int32 seed);
+	FOLIAGE_API void AddInstance(AInstancedFoliageActor* InIFA, const UFoliageType* InSettings, const FFoliageInstance& InNewInstance);
+	FOLIAGE_API void AddInstance(AInstancedFoliageActor* InIFA, const UFoliageType* InSettings, const FFoliageInstance& InNewInstance, UActorComponent* InBaseComponent);
+	FOLIAGE_API void AddInstances(AInstancedFoliageActor* InIFA, const UFoliageType* InSettings, const TSet<const FFoliageInstance*>& InNewInstances);
 
-#if WITH_EDITOR
-	FOLIAGE_API void AddInstance(AInstancedFoliageActor* InIFA, const UFoliageType* InSettings, const FFoliageInstance& InNewInstance, bool RebuildFoliageTree);
-	FOLIAGE_API void AddInstance(AInstancedFoliageActor* InIFA, const UFoliageType* InSettings, const FFoliageInstance& InNewInstance, UActorComponent* InBaseComponent, bool RebuildFoliageTree);
-	FOLIAGE_API void AddInstances(AInstancedFoliageActor* InIFA, const UFoliageType* InSettings, const TSet<const FFoliageInstance*>& InNewInstances, bool RebuildFoliageTree);
-	
 	FOLIAGE_API void RemoveInstances(AInstancedFoliageActor* InIFA, const TArray<int32>& InInstancesToRemove, bool RebuildFoliageTree);
 	// Apply changes in the FoliageType to the component
-	FOLIAGE_API void UpdateComponentSettings(const UFoliageType* InSettings);
-	// Recreate the component if the FoliageType's ComponentClass doesn't match the Component's class
-	FOLIAGE_API void CheckComponentClass(AInstancedFoliageActor* InIFA, const UFoliageType* InSettings);
 	FOLIAGE_API void PreMoveInstances(AInstancedFoliageActor* InIFA, const TArray<int32>& InInstancesToMove);
 	FOLIAGE_API void PostMoveInstances(AInstancedFoliageActor* InIFA, const TArray<int32>& InInstancesMoved);
 	FOLIAGE_API void PostUpdateInstances(AInstancedFoliageActor* InIFA, const TArray<int32>& InInstancesUpdated, bool bReAddToHash = false, bool InUpdateSelection = false);
@@ -214,12 +285,11 @@ struct FFoliageMeshInfo
 	FOLIAGE_API void GetInstanceAtLocation(const FVector& Location, int32& OutInstance, bool& bOutSucess);
 	FOLIAGE_API bool CheckForOverlappingSphere(const FSphere& Sphere);
 	FOLIAGE_API bool CheckForOverlappingInstanceExcluding(int32 TestInstanceIdx, float Radius, TSet<int32>& ExcludeInstances);
+	FOLIAGE_API TArray<int32> GetInstancesOverlappingBox(const FBox& Box) const;
 
 	// Destroy existing clusters and reassign all instances to new clusters
 	FOLIAGE_API void ReallocateClusters(AInstancedFoliageActor* InIFA, UFoliageType* InSettings);
 
-	FOLIAGE_API void ReapplyInstancesToComponent();
-
 	FOLIAGE_API void SelectInstances(AInstancedFoliageActor* InIFA, bool bSelect, TArray<int32>& Instances);
 
 	FOLIAGE_API void SelectInstances(AInstancedFoliageActor* InIFA, bool bSelect);
@@ -230,22 +300,28 @@ struct FFoliageMeshInfo
 	FOLIAGE_API void AddToBaseHash(int32 InstanceIdx);
 	FOLIAGE_API void RemoveFromBaseHash(int32 InstanceIdx);
 
-	// Create and register a new component
-	FOLIAGE_API void CreateNewComponent(AInstancedFoliageActor* InIFA, const UFoliageType* InSettings);
-
 	// For debugging. Validate state after editing.
 	void CheckValid();
-
-	FOLIAGE_API void HandleComponentMeshBoundsChanged(const FBoxSphereBounds& NewBounds);
+		
+	FOLIAGE_API void Refresh(AInstancedFoliageActor* IFA, bool Async, bool Force);
+	FOLIAGE_API void OnHiddenEditorViewMaskChanged(uint64 InHiddenEditorViews);
+	FOLIAGE_API void PostEditUndo(AInstancedFoliageActor* IFA, UFoliageType* FoliageType);
+	FOLIAGE_API void PreEditUndo(AInstancedFoliageActor* IFA, UFoliageType* FoliageType);
+	FOLIAGE_API void EnterEditMode();
+	FOLIAGE_API void ExitEditMode();
 #endif
 
-	friend FArchive& operator<<(FArchive& Ar, FFoliageMeshInfo& MeshInfo);
+	friend FArchive& operator<<(FArchive& Ar, FFoliageInfo& MeshInfo);
 
 	// Non-copyable
-	FFoliageMeshInfo(const FFoliageMeshInfo&) = delete;
-	FFoliageMeshInfo& operator=(const FFoliageMeshInfo&) = delete;
+	FFoliageInfo(const FFoliageInfo&) = delete;
+	FFoliageInfo& operator=(const FFoliageInfo&) = delete;
+
+private:
+	void AddInstanceImpl(AInstancedFoliageActor* InIFA, const FFoliageInstance& InNewInstance);
 };
 
+
 #if WITH_EDITORONLY_DATA
 //
 // FFoliageInstanceHash
@@ -317,7 +393,6 @@ public:
 		return Result;
 	}
 
-#if UE_BUILD_DEBUG
 	void CheckInstanceCount(int32 InCount) const
 	{
 		int32 HashCount = 0;
@@ -328,7 +403,6 @@ public:
 
 		check(HashCount == InCount);
 	}
-#endif
 
 	void Empty()
 	{
diff --git a/Engine/Source/Runtime/Foliage/Public/InstancedFoliageActor.h b/Engine/Source/Runtime/Foliage/Public/InstancedFoliageActor.h
--- a/Engine/Source/Runtime/Foliage/Public/InstancedFoliageActor.h
+++ b/Engine/Source/Runtime/Foliage/Public/InstancedFoliageActor.h
@@ -14,10 +14,6 @@
 #include "InstancedFoliageActor.generated.h"
 
 class UProceduralFoliageComponent;
-struct FDesiredFoliageInstance;
-struct FFoliageInstance;
-struct FFoliageInstancePlacementInfo;
-struct FFoliageMeshInfo;
 
 // Function for filtering out hit components during FoliageTrace
 typedef TFunction<bool(const UPrimitiveComponent*)> FFoliageTraceFilterFunc;
@@ -33,7 +29,7 @@ public:
 	FFoliageInstanceBaseCache InstanceBaseCache;
 #endif// WITH_EDITORONLY_DATA
 
-	TMap<UFoliageType*, TUniqueObj<FFoliageMeshInfo>> FoliageMeshes;
+	TMap<UFoliageType*, TUniqueObj<FFoliageInfo>> FoliageInfos;
 
 public:
 	//~ Begin UObject Interface.
@@ -46,6 +42,8 @@ public:
 	// we don't want to have our components automatically destroyed by the Blueprint code
 	virtual void RerunConstructionScripts() override {}
 	virtual bool IsLevelBoundsRelevant() const override { return false; }
+
+	FOLIAGE_API static bool IsOwnedByFoliage(const AActor* InActor);
 protected:
 	// Default InternalTakeRadialDamage behavior finds and scales damage for the closest component which isn't appropriate for foliage.
 	virtual float InternalTakeRadialDamage(float Damage, struct FRadialDamageEvent const& RadialDamageEvent, class AController* EventInstigator, AActor* DamageCauser) override;
@@ -55,17 +53,21 @@ public:
 	virtual void BeginDestroy() override;
 	virtual void Destroyed() override;
 	FOLIAGE_API void CleanupDeletedFoliageType();
+
+	// Delegate type for selection change events
+	DECLARE_MULTICAST_DELEGATE_TwoParams(FOnSelectionChanged, bool, const TArray<AActor*>&);
+	FOLIAGE_API static FOnSelectionChanged SelectionChanged;
 #endif
 	//~ End AActor Interface.
 
 
-	// Performs a reverse lookup from a mesh to a local foliage type (i.e. the foliage type owned exclusively by this IFA)
-	FOLIAGE_API UFoliageType* GetLocalFoliageTypeForMesh(const UStaticMesh* InMesh, FFoliageMeshInfo** OutMeshInfo = nullptr);
+	// Performs a reverse lookup from a source object to a local foliage type (i.e. the foliage type owned exclusively by this IFA)
+	FOLIAGE_API UFoliageType* GetLocalFoliageTypeForSource(const UObject* InSource, FFoliageInfo** OutMeshInfo = nullptr);
 	
-	// Performs a reverse lookup from a mesh to all the foliage types that are currently using that mesh (includes assets and blueprint classes)
-	FOLIAGE_API void GetAllFoliageTypesForMesh(const UStaticMesh* InMesh, TArray<const UFoliageType*>& OutFoliageTypes);
+	// Performs a reverse lookup from a source object to all the foliage types that are currently using that object (includes assets and blueprint classes)
+	FOLIAGE_API void GetAllFoliageTypesForSource(const UObject* InSource, TArray<const UFoliageType*>& OutFoliageTypes);
 	
-	FOLIAGE_API FFoliageMeshInfo* FindFoliageTypeOfClass(TSubclassOf<UFoliageType_InstancedStaticMesh> Class);
+	FOLIAGE_API FFoliageInfo* FindFoliageTypeOfClass(TSubclassOf<UFoliageType_InstancedStaticMesh> Class);
 
 	// Finds the number of instances overlapping with the sphere. 
 	FOLIAGE_API int32 GetOverlappingSphereCount(const UFoliageType* FoliageType, const FSphere& Sphere) const;
@@ -78,10 +80,10 @@ public:
 	FOLIAGE_API void GetOverlappingMeshCounts(const FSphere& Sphere, TMap<UStaticMesh*, int32>& OutCounts) const;
 
 	// Finds a mesh entry
-	FOLIAGE_API FFoliageMeshInfo* FindMesh(const UFoliageType* InType);
+	FOLIAGE_API FFoliageInfo* FindInfo(const UFoliageType* InType);
 
 	// Finds a mesh entry
-	FOLIAGE_API const FFoliageMeshInfo* FindMesh(const UFoliageType* InType) const;
+	FOLIAGE_API const FFoliageInfo* FindInfo(const UFoliageType* InType) const;
 
 	/**
 	* Get the instanced foliage actor for the current streaming level.
@@ -148,12 +150,12 @@ public:
 	bool ContainsInstancesFromProceduralFoliageComponent(const UProceduralFoliageComponent* ProceduralFoliageComponent);
 
 	// Finds a mesh entry or adds it if it doesn't already exist
-	FOLIAGE_API FFoliageMeshInfo* FindOrAddMesh(UFoliageType* InType);
+	FOLIAGE_API FFoliageInfo* FindOrAddMesh(UFoliageType* InType);
 
-	FOLIAGE_API UFoliageType* AddFoliageType(const UFoliageType* InType, FFoliageMeshInfo** OutInfo = nullptr);
+	FOLIAGE_API UFoliageType* AddFoliageType(const UFoliageType* InType, FFoliageInfo** OutInfo = nullptr);
 	// Add a new static mesh.
-	FOLIAGE_API FFoliageMeshInfo* AddMesh(UStaticMesh* InMesh, UFoliageType** OutSettings = nullptr, const UFoliageType_InstancedStaticMesh* DefaultSettings = nullptr);
-	FOLIAGE_API FFoliageMeshInfo* AddMesh(UFoliageType* InType);
+	FOLIAGE_API FFoliageInfo* AddMesh(UStaticMesh* InMesh, UFoliageType** OutSettings = nullptr, const UFoliageType_InstancedStaticMesh* DefaultSettings = nullptr);
+	FOLIAGE_API FFoliageInfo* AddMesh(UFoliageType* InType);
 
 	// Remove the FoliageType from the list, and all its instances.
 	FOLIAGE_API void RemoveFoliageType(UFoliageType** InFoliageType, int32 Num);
@@ -161,17 +163,20 @@ public:
 	// Select an individual instance.
 	FOLIAGE_API void SelectInstance(UInstancedStaticMeshComponent* InComponent, int32 InComponentInstanceIndex, bool bToggle);
 
+	// Select an individual instance.
+	FOLIAGE_API void SelectInstance(AActor* InActor, bool bToggle);
+
 	// Whether actor has selected instances
 	FOLIAGE_API bool HasSelectedInstances() const;
 
 	// Will return all the foliage type used by currently selected instances
-	FOLIAGE_API TMap<UFoliageType*, FFoliageMeshInfo*> GetSelectedInstancesFoliageType();
+	FOLIAGE_API TMap<UFoliageType*, FFoliageInfo*> GetSelectedInstancesFoliageType();
 
 	// Will return all the foliage type used
-	FOLIAGE_API TMap<UFoliageType*, FFoliageMeshInfo*> GetAllInstancesFoliageType();
+	FOLIAGE_API TMap<UFoliageType*, FFoliageInfo*> GetAllInstancesFoliageType();
 
-	// Propagate the selected instances to the actual render components
-	FOLIAGE_API void ApplySelectionToComponents(bool bApply);
+	// Propagate the selected instances to the actual foliage implementation
+	FOLIAGE_API void ApplySelection(bool bApply);
 
 	// Returns the location for the widget
 	FOLIAGE_API bool GetSelectionLocation(FVector& OutLocation) const;
@@ -180,8 +185,8 @@ public:
 	static FOLIAGE_API bool HasFoliageAttached(UActorComponent* InComponent);
 
 	/* Called to notify InstancedFoliageActor that a UFoliageType has been modified */
-	void NotifyFoliageTypeChanged(UFoliageType* FoliageType, bool bMeshChanged);
-	void NotifyFoliageTypeWillChange(UFoliageType* FoliageType, bool bMeshChanged);
+	void NotifyFoliageTypeChanged(UFoliageType* FoliageType, bool bSourceChanged);
+	void NotifyFoliageTypeWillChange(UFoliageType* FoliageType);
 
 	DECLARE_EVENT_OneParam(AInstancedFoliageActor, FOnFoliageTypeMeshChanged, UFoliageType*);
 	FOnFoliageTypeMeshChanged& OnFoliageTypeMeshChanged() { return OnFoliageTypeMeshChangedEvent; }
@@ -194,6 +199,7 @@ private:
 #if WITH_EDITORONLY_DATA
 	// Deprecated data, will be converted and cleaned up in PostLoad
 	TMap<UFoliageType*, TUniqueObj<struct FFoliageMeshInfo_Deprecated>> FoliageMeshes_Deprecated;
+	TMap<UFoliageType*, TUniqueObj<struct FFoliageMeshInfo_Deprecated2>> FoliageMeshes_Deprecated2;
 #endif//WITH_EDITORONLY_DATA
 	
 #if WITH_EDITOR
@@ -202,7 +208,7 @@ private:
 	void OnPostApplyLevelOffset(ULevel* InLevel, UWorld* InWorld, const FVector& InOffset, bool bWorldShift);
 
 	// Move instances to a foliage actor in target level
-	FOLIAGE_API void MoveInstancesToLevel(ULevel* InTargetLevel, TSet<int32>& InInstanceList, FFoliageMeshInfo* InCurrentMeshInfo, UFoliageType* InFoliageType);
+	FOLIAGE_API void MoveInstancesToLevel(ULevel* InTargetLevel, TSet<int32>& InInstanceList, FFoliageInfo* InCurrentMeshInfo, UFoliageType* InFoliageType);
 #endif
 private:
 #if WITH_EDITOR
diff --git a/Engine/Source/Runtime/Foliage/Public/ProceduralFoliageInstance.h b/Engine/Source/Runtime/Foliage/Public/ProceduralFoliageInstance.h
--- a/Engine/Source/Runtime/Foliage/Public/ProceduralFoliageInstance.h
+++ b/Engine/Source/Runtime/Foliage/Public/ProceduralFoliageInstance.h
@@ -8,7 +8,7 @@
 #include "ProceduralFoliageInstance.generated.h"
 
 class UActorComponent;
-class UFoliageType_InstancedStaticMesh;
+class UFoliageType;
 struct FProceduralFoliageInstance;
 
 #if WITH_PHYSX
@@ -95,7 +95,7 @@ public:
 	float Scale;
 
 	UPROPERTY()
-	const UFoliageType_InstancedStaticMesh* Type;
+	const UFoliageType* Type;
 
 
 	bool bBlocker;
diff --git a/Engine/Source/Runtime/Foliage/Public/ProceduralFoliageTile.h b/Engine/Source/Runtime/Foliage/Public/ProceduralFoliageTile.h
--- a/Engine/Source/Runtime/Foliage/Public/ProceduralFoliageTile.h
+++ b/Engine/Source/Runtime/Foliage/Public/ProceduralFoliageTile.h
@@ -11,7 +11,7 @@
 #include "ProceduralFoliageBroadphase.h"
 #include "ProceduralFoliageTile.generated.h"
 
-class UFoliageType_InstancedStaticMesh;
+class UFoliageType;
 class UProceduralFoliageSpawner;
 struct FBodyInstance;
 
@@ -110,7 +110,7 @@ private:
 	bool HandleOverlaps(FProceduralFoliageInstance* Instance);
 
 	/** Attempts to create a new instance and resolves any overlaps. Returns the new instance if successful for calling code to add to Instances */
-	FProceduralFoliageInstance* NewSeed(const FVector& Location, float Scale, const UFoliageType_InstancedStaticMesh* Type, float InAge, bool bBlocker = false);
+	FProceduralFoliageInstance* NewSeed(const FVector& Location, float Scale, const UFoliageType* Type, float InAge, bool bBlocker = false);
 
 	void SpreadSeeds(TArray<FProceduralFoliageInstance*>& NewSeeds);
 	void AgeSeeds();
@@ -142,5 +142,5 @@ private:
 
 private:
 	float GetRandomGaussian();
-	FVector GetSeedOffset(const UFoliageType_InstancedStaticMesh* Type, float MinDistance);
+	FVector GetSeedOffset(const UFoliageType* Type, float MinDistance);
 };
diff --git a/Engine/Source/Runtime/Landscape/Private/LandscapeCollision.cpp b/Engine/Source/Runtime/Landscape/Private/LandscapeCollision.cpp
--- a/Engine/Source/Runtime/Landscape/Private/LandscapeCollision.cpp
+++ b/Engine/Source/Runtime/Landscape/Private/LandscapeCollision.cpp
@@ -1246,11 +1246,11 @@ void ULandscapeHeightfieldCollisionComponent::SnapFoliageInstances(const FBox& I
 			continue;
 		}
 			
-		for (auto& MeshPair : IFA->FoliageMeshes)
+		for (auto& Pair : IFA->FoliageInfos)
 		{
 			// Find the per-mesh info matching the mesh.
-			UFoliageType* Settings = MeshPair.Key;
-			FFoliageMeshInfo& MeshInfo = *MeshPair.Value;
+			UFoliageType* Settings = Pair.Key;
+			FFoliageInfo& MeshInfo = *Pair.Value;
 			
 			const auto* InstanceSet = MeshInfo.ComponentHash.Find(BaseId);
 			if (InstanceSet)
@@ -1260,7 +1260,7 @@ void ULandscapeHeightfieldCollisionComponent::SnapFoliageInstances(const FBox& I
 
 				bool bFirst = true;
 				TArray<int32> InstancesToRemove;
-				TSet<UHierarchicalInstancedStaticMeshComponent*> AffectedFoliageComponets;
+				TSet<UHierarchicalInstancedStaticMeshComponent*> AffectedFoliageComponents;
 
 				for (int32 InstanceIndex : *InstanceSet)
 				{
@@ -1319,13 +1319,9 @@ void ULandscapeHeightfieldCollisionComponent::SnapFoliageInstances(const FBox& I
 
 									// Todo: add do validation with other parameters such as max/min height etc.
 
-									check(MeshInfo.Component);
-									MeshInfo.Component->Modify();
-									MeshInfo.Component->UpdateInstanceTransform(InstanceIndex, Instance.GetInstanceWorldTransform(), true, false);
+									MeshInfo.SetInstanceWorldTransform(InstanceIndex, Instance.GetInstanceWorldTransform(), false);
 									// Re-add the new instance location to the hash
 									MeshInfo.InstanceHash->InsertInstance(Instance.Location, InstanceIndex);
-
-									MeshInfo.Component->MarkRenderStateDirty();
 								}
 								break;
 							}
@@ -1337,14 +1333,18 @@ void ULandscapeHeightfieldCollisionComponent::SnapFoliageInstances(const FBox& I
 							InstancesToRemove.Add(InstanceIndex);
 						}
 
-						AffectedFoliageComponets.Add(MeshInfo.Component);
+						if (MeshInfo.GetComponent() != nullptr)
+						{
+							AffectedFoliageComponents.Add(MeshInfo.GetComponent());
+						}
+						
 					}
 				}
 
 				// Remove any unused instances
 				MeshInfo.RemoveInstances(IFA, InstancesToRemove, true);
 
-				for (UHierarchicalInstancedStaticMeshComponent* FoliageComp : AffectedFoliageComponets)
+				for (UHierarchicalInstancedStaticMeshComponent* FoliageComp : AffectedFoliageComponents)
 				{
 					FoliageComp->InvalidateLightingCache();
 				}
