Fix bad merge

#rnx

diff --git a/Engine/Source/Runtime/Foliage/Private/InstancedFoliage.cpp b/Engine/Source/Runtime/Foliage/Private/InstancedFoliage.cpp
--- a/Engine/Source/Runtime/Foliage/Private/InstancedFoliage.cpp
+++ b/Engine/Source/Runtime/Foliage/Private/InstancedFoliage.cpp
@@ -1972,46 +1972,7 @@ void FFoliageInfo::ReallocateClusters(AInstancedFoliageActor* InIFA, UFoliageTyp
 	Refresh(InIFA, true, true);
 }
 
-void FFoliageMeshInfo::ReapplyInstancesToComponent()
-{
-	if (Component)
-	{
-		// clear the transactional flag if it was set prior to deleting the actor
-		Component->ClearFlags(RF_Transactional);
-
-		const bool bWasRegistered = Component->IsRegistered();
-		Component->UnregisterComponent();
-		Component->ClearInstances();
-		Component->InitPerInstanceRenderData(false);
-
-		Component->bAutoRebuildTreeOnInstanceChanges = false;
-
-		for (auto& Instance : Instances)
-		{
-			Component->AddInstanceWorldSpace(Instance.GetInstanceWorldTransform());
-		}
-
-		Component->bAutoRebuildTreeOnInstanceChanges = true;
-		Component->BuildTreeIfOutdated(true, true);
-
-		Component->ClearInstanceSelection();
-
-		if (SelectedIndices.Num())
-		{
-			for (int32 i : SelectedIndices)
-			{
-				Component->SelectInstance(true, i, 1);
-			}
-		}
-
-		if (bWasRegistered)
-		{
-			Component->RegisterComponent();
-		}
-	}
-}
-
-void FFoliageMeshInfo::GetInstancesInsideSphere(const FSphere& Sphere, TArray<int32>& OutInstances)
+void FFoliageInfo::GetInstancesInsideSphere(const FSphere& Sphere, TArray<int32>& OutInstances)
 {
 	auto TempInstances = InstanceHash->GetInstancesOverlappingBox(FBox::BuildAABB(Sphere.Center, FVector(Sphere.W)));
 	for (int32 Idx : TempInstances)
