- Fix New Foliage Type Crash in Multi-level map
- Fix Replace Foliage Type

diff --git a/Engine/Source/Editor/FoliageEdit/Private/FoliageEdMode.cpp b/Engine/Source/Editor/FoliageEdit/Private/FoliageEdMode.cpp
--- a/Engine/Source/Editor/FoliageEdit/Private/FoliageEdMode.cpp
+++ b/Engine/Source/Editor/FoliageEdit/Private/FoliageEdMode.cpp
@@ -813,7 +813,10 @@ void FEdModeFoliage::OnObjectsReplaced(const TMap<UObject*, UObject*>& Replaceme
 						if (IFA->FoliageInfos.RemoveAndCopyValue(ReplacedFoliageType, FoliageInfo))
 						{
 							// Re-add the unique mesh info associated with the replaced foliage type
-							IFA->FoliageInfos.Add(Cast<UFoliageType>(ReplacementPair.Value), MoveTemp(FoliageInfo));
+							UFoliageType* ReplacementFoliageType = Cast<UFoliageType>(ReplacementPair.Value);
+							TUniqueObj<FFoliageInfo>& NewFoliageInfo = IFA->FoliageInfos.Add(ReplacementFoliageType, MoveTemp(FoliageInfo));
+							NewFoliageInfo->ReallocateClusters(IFA, ReplacementFoliageType);
+							
 							bAnyFoliageTypeReplaced = true;
 						}
 					}
diff --git a/Engine/Source/Editor/FoliageEdit/Private/FoliageEditUtility.cpp b/Engine/Source/Editor/FoliageEdit/Private/FoliageEditUtility.cpp
--- a/Engine/Source/Editor/FoliageEdit/Private/FoliageEditUtility.cpp
+++ b/Engine/Source/Editor/FoliageEdit/Private/FoliageEditUtility.cpp
@@ -110,7 +110,9 @@ void FFoliageEditUtility::ReplaceFoliageTypeObject(UWorld* InWorld, UFoliageType
 				}
 				else
 				{
-					IFA->FoliageInfos.Add(NewType, MoveTemp(OldInfo))->ReallocateClusters(IFA, NewType);
+					// Make sure if type changes we have proper implementation
+					TUniqueObj<FFoliageInfo>& NewFoliageInfo = IFA->FoliageInfos.Add(NewType, MoveTemp(OldInfo));
+					NewFoliageInfo->ReallocateClusters(IFA, NewType);
 				}
 			}
 		}
diff --git a/Engine/Source/Runtime/Foliage/Private/InstancedFoliage.cpp b/Engine/Source/Runtime/Foliage/Private/InstancedFoliage.cpp
--- a/Engine/Source/Runtime/Foliage/Private/InstancedFoliage.cpp
+++ b/Engine/Source/Runtime/Foliage/Private/InstancedFoliage.cpp
@@ -1945,10 +1945,10 @@ void FFoliageInfo::RemoveFromBaseHash(int32 InstanceIndex)
 // Destroy existing clusters and reassign all instances to new clusters
 void FFoliageInfo::ReallocateClusters(AInstancedFoliageActor* InIFA, UFoliageType* InSettings)
 {
-	if (Implementation->IsInitialized())
-	{
-		Implementation->Uninitialize();
-	}
+	// In case Foliage Type Changed recreate implementation
+	Implementation.Reset();
+	CreateImplementation(InSettings);
+	
 	// Remove everything
 	TArray<FFoliageInstance> OldInstances;
 	Exchange(Instances, OldInstances);
@@ -2823,6 +2823,10 @@ FFoliageInfo* AInstancedFoliageActor::AddMesh(UFoliageType* InType)
 	Modify();
 
 	FFoliageInfo* Info = &*FoliageInfos.Add(InType);
+	if (Info && !Info->Implementation.IsValid())
+	{
+		Info->CreateImplementation(InType);
+	}
 	Info->FoliageTypeUpdateGuid = InType->UpdateGuid;
 	InType->IsSelected = true;
 
