- Fix StaticAnalysis warning caused by useless null check

#rnx

diff --git a/Engine/Source/Runtime/Foliage/Private/InstancedFoliage.cpp b/Engine/Source/Runtime/Foliage/Private/InstancedFoliage.cpp
--- a/Engine/Source/Runtime/Foliage/Private/InstancedFoliage.cpp
+++ b/Engine/Source/Runtime/Foliage/Private/InstancedFoliage.cpp
@@ -2823,7 +2823,7 @@ FFoliageInfo* AInstancedFoliageActor::AddMesh(UFoliageType* InType)
 	Modify();
 
 	FFoliageInfo* Info = &*FoliageInfos.Add(InType);
-	if (Info && !Info->Implementation.IsValid())
+	if (!Info->Implementation.IsValid())
 	{
 		Info->CreateImplementation(InType);
 	}
