fix for crash on startup for IOS and Android

diff --git a/Engine/Source/Runtime/Foliage/Private/InstancedFoliage.cpp b/Engine/Source/Runtime/Foliage/Private/InstancedFoliage.cpp
--- a/Engine/Source/Runtime/Foliage/Private/InstancedFoliage.cpp
+++ b/Engine/Source/Runtime/Foliage/Private/InstancedFoliage.cpp
@@ -1541,7 +1541,10 @@ UHierarchicalInstancedStaticMeshComponent* FFoliageInfo::GetComponent() const
 
 void FFoliageInfo::AddReferencedObjects(UObject* InThis, FReferenceCollector& Collector)
 {
-	Implementation->AddReferencedObjects(InThis, Collector);
+	if (Implementation.IsValid())
+	{
+		Implementation->AddReferencedObjects(InThis, Collector);
+	}
 }
 
 void FFoliageInfo::CreateImplementation(EFoliageImplType InType)
