Fix for when CVarFoliageDiscardDataOnLoad is set to 1. (low end devices from what I understand)

#rnx

diff --git a/Engine/Source/Runtime/Foliage/Private/InstancedFoliage.cpp b/Engine/Source/Runtime/Foliage/Private/InstancedFoliage.cpp
--- a/Engine/Source/Runtime/Foliage/Private/InstancedFoliage.cpp
+++ b/Engine/Source/Runtime/Foliage/Private/InstancedFoliage.cpp
@@ -1530,7 +1530,7 @@ FFoliageInfo::~FFoliageInfo()
 
 UHierarchicalInstancedStaticMeshComponent* FFoliageInfo::GetComponent() const
 {
-	if (Type == EFoliageImplType::StaticMesh)
+	if (Type == EFoliageImplType::StaticMesh && Implementation.IsValid())
 	{
 		FFoliageStaticMesh* FoliageStaticMesh = StaticCast<FFoliageStaticMesh*>(Implementation.Get());
 		return FoliageStaticMesh->Component;
@@ -1566,22 +1566,38 @@ void FFoliageInfo::CreateImplementation(EFoliageImplType InType)
 
 int32 FFoliageInfo::GetOverlappingSphereCount(const FSphere& Sphere) const
 {
-	return Implementation->GetOverlappingSphereCount(Sphere);
+	if (Implementation.IsValid())
+	{
+		return Implementation->GetOverlappingSphereCount(Sphere);
+	}
+
+	return 0;
 }
 
 int32 FFoliageInfo::GetOverlappingBoxCount(const FBox& Box) const
 {
-	return Implementation->GetOverlappingBoxCount(Box);
+	if (Implementation.IsValid())
+	{
+		return Implementation->GetOverlappingBoxCount(Box);
+	}
+
+	return 0;
 }
 
 void FFoliageInfo::GetOverlappingBoxTransforms(const FBox& Box, TArray<FTransform>& OutTransforms) const
 {
-	Implementation->GetOverlappingBoxTransforms(Box, OutTransforms);
+	if (Implementation.IsValid())
+	{
+		Implementation->GetOverlappingBoxTransforms(Box, OutTransforms);
+	}
 }
 
 void FFoliageInfo::GetOverlappingMeshCount(const FSphere& Sphere, TMap<UStaticMesh*, int32>& OutCounts) const
 {
-	Implementation->GetOverlappingMeshCount(Sphere, OutCounts);
+	if (Implementation.IsValid())
+	{
+		Implementation->GetOverlappingMeshCount(Sphere, OutCounts);
+	}
 }
 
 #if WITH_EDITOR
