Fixed foliage selection not clearing selection of all Instance Foliage Actors

diff --git a/Engine/Source/Runtime/Foliage/Private/InstancedFoliage.cpp b/Engine/Source/Runtime/Foliage/Private/InstancedFoliage.cpp
--- a/Engine/Source/Runtime/Foliage/Private/InstancedFoliage.cpp
+++ b/Engine/Source/Runtime/Foliage/Private/InstancedFoliage.cpp
@@ -2876,6 +2876,26 @@ void AInstancedFoliageActor::RemoveFoliageType(UFoliageType** InFoliageTypes, in
 	RegisterAllComponents();
 }
 
+void AInstancedFoliageActor::ClearSelection()
+{
+	UWorld* World = GetWorld();
+	const int32 NumLevels = World->GetNumLevels();
+	for (int32 LevelIdx = 0; LevelIdx < NumLevels; ++LevelIdx)
+	{
+		if (ULevel* Level = World->GetLevel(LevelIdx))
+		{
+			if (AInstancedFoliageActor* IFA = AInstancedFoliageActor::GetInstancedFoliageActorForLevel(Level))
+			{
+				for (auto& Pair : IFA->FoliageInfos)
+				{
+					FFoliageInfo& Info = *Pair.Value;
+					Info.ClearSelection();
+				}
+			}
+		}
+	}
+}
+
 void AInstancedFoliageActor::SelectInstance(UInstancedStaticMeshComponent* InComponent, int32 InInstanceIndex, bool bToggle)
 {
 	Modify();
@@ -2883,11 +2903,7 @@ void AInstancedFoliageActor::SelectInstance(UInstancedStaticMeshComponent* InCom
 	// If we're not toggling, we need to first deselect everything else
 	if (!bToggle)
 	{
-		for (auto& Pair : FoliageInfos)
-		{
-			FFoliageInfo& Info = *Pair.Value;
-			Info.ClearSelection();
-		}
+		ClearSelection();
 	}
 
 	if (InComponent)
@@ -2944,11 +2960,7 @@ void AInstancedFoliageActor::SelectInstance(AActor* InActor, bool bToggle)
 	// If we're not toggling, we need to first deselect everything else
 	if (!bToggle)
 	{
-		for (auto& Pair : FoliageInfos)
-		{
-			FFoliageInfo& Info = *Pair.Value;
-			Info.ClearSelection();
-		}
+		ClearSelection();
 	}
 
 	if (InActor)
diff --git a/Engine/Source/Runtime/Foliage/Public/InstancedFoliageActor.h b/Engine/Source/Runtime/Foliage/Public/InstancedFoliageActor.h
--- a/Engine/Source/Runtime/Foliage/Public/InstancedFoliageActor.h
+++ b/Engine/Source/Runtime/Foliage/Public/InstancedFoliageActor.h
@@ -203,6 +203,7 @@ private:
 #endif//WITH_EDITORONLY_DATA
 	
 #if WITH_EDITOR
+	void ClearSelection();
 	void OnLevelActorMoved(AActor* InActor);
 	void OnLevelActorDeleted(AActor* InActor);
 	void OnPostApplyLevelOffset(ULevel* InLevel, UWorld* InWorld, const FVector& InOffset, bool bWorldShift);
