Duplicate Fix (CL 7039627)

#rnx

diff --git a/Engine/Source/Runtime/Foliage/Private/FoliageActor.cpp b/Engine/Source/Runtime/Foliage/Private/FoliageActor.cpp
--- a/Engine/Source/Runtime/Foliage/Private/FoliageActor.cpp
+++ b/Engine/Source/Runtime/Foliage/Private/FoliageActor.cpp
@@ -44,24 +44,22 @@ void FFoliageActor::DestroyActors(bool bOnLoad)
 #if WITH_EDITOR
 bool FFoliageActor::IsInitialized() const
 {
-	return bInitialized;
+	return ActorClass != nullptr;;
 }
 
 void FFoliageActor::Initialize(AInstancedFoliageActor* IFA, const UFoliageType* FoliageType)
 {
-	check(!bInitialized);
+	check(!IsInitialized());
 	const UFoliageType_Actor* FoliageType_Actor = Cast<UFoliageType_Actor>(FoliageType);
-	ActorClass = FoliageType_Actor->ActorClass;
-	bInitialized = true;
+	ActorClass = FoliageType_Actor->ActorClass ? FoliageType_Actor->ActorClass : AActor::StaticClass();
 }
 
 void FFoliageActor::Uninitialize()
 {
-	check(bInitialized);
+	check(IsInitialized());
 	DestroyActors(false);
 	ActorInstances.Empty();
 	ActorClass = nullptr;
-	bInitialized = false;
 }
 
 AActor* FFoliageActor::Spawn(AInstancedFoliageActor* IFA, const FFoliageInstance& Instance)
diff --git a/Engine/Source/Runtime/Foliage/Private/FoliageActor.h b/Engine/Source/Runtime/Foliage/Private/FoliageActor.h
--- a/Engine/Source/Runtime/Foliage/Private/FoliageActor.h
+++ b/Engine/Source/Runtime/Foliage/Private/FoliageActor.h
@@ -7,11 +7,9 @@ struct FFoliageActor : public FFoliageImpl
 {
 	TArray<TWeakObjectPtr<AActor>> ActorInstances;
 	UClass* ActorClass;
-	bool bInitialized;
-
+	
 	FFoliageActor()
 		: ActorClass(nullptr)
-		, bInitialized(false)
 	{
 	}
 
