Foliage: Check for valid level in PostLoad()

diff --git a/Engine/Source/Runtime/Foliage/Private/InstancedFoliage.cpp b/Engine/Source/Runtime/Foliage/Private/InstancedFoliage.cpp
--- a/Engine/Source/Runtime/Foliage/Private/InstancedFoliage.cpp
+++ b/Engine/Source/Runtime/Foliage/Private/InstancedFoliage.cpp
@@ -3408,24 +3408,27 @@ void AInstancedFoliageActor::PostLoad()
 	Super::PostLoad();
 
 	ULevel* OwningLevel = GetLevel();
-	if (!OwningLevel->InstancedFoliageActor.IsValid())
+	if (OwningLevel)
 	{
-		OwningLevel->InstancedFoliageActor = this;
-	}
-	else
-	{
-		FFormatNamedArguments Arguments;
-		Arguments.Add(TEXT("Level"), FText::FromString(*OwningLevel->GetOutermost()->GetName()));
-		FMessageLog("MapCheck").Warning()
-			->AddToken(FUObjectToken::Create(this))
-			->AddToken(FTextToken::Create(FText::Format(LOCTEXT("MapCheck_DuplicateInstancedFoliageActor", "Level {Level} has an unexpected duplicate Instanced Foliage Actor."), Arguments)))
+		if (!OwningLevel->InstancedFoliageActor.IsValid())
+		{
+			OwningLevel->InstancedFoliageActor = this;
+		}
+		else
+		{
+			FFormatNamedArguments Arguments;
+			Arguments.Add(TEXT("Level"), FText::FromString(*OwningLevel->GetOutermost()->GetName()));
+			FMessageLog("MapCheck").Warning()
+				->AddToken(FUObjectToken::Create(this))
+				->AddToken(FTextToken::Create(FText::Format(LOCTEXT("MapCheck_DuplicateInstancedFoliageActor", "Level {Level} has an unexpected duplicate Instanced Foliage Actor."), Arguments)))
 #if WITH_EDITOR
-			->AddToken(FActionToken::Create(LOCTEXT("MapCheck_FixDuplicateInstancedFoliageActor", "Fix"),
-				LOCTEXT("MapCheck_FixDuplicateInstancedFoliageActor_Desc", "Click to consolidate foliage into the main foliage actor."),
-				FOnActionTokenExecuted::CreateUObject(OwningLevel->InstancedFoliageActor.Get(), &AInstancedFoliageActor::RepairDuplicateIFA, this), true))
+				->AddToken(FActionToken::Create(LOCTEXT("MapCheck_FixDuplicateInstancedFoliageActor", "Fix"),
+					LOCTEXT("MapCheck_FixDuplicateInstancedFoliageActor_Desc", "Click to consolidate foliage into the main foliage actor."),
+					FOnActionTokenExecuted::CreateUObject(OwningLevel->InstancedFoliageActor.Get(), &AInstancedFoliageActor::RepairDuplicateIFA, this), true))
 #endif// WITH_EDITOR
-			;
-		FMessageLog("MapCheck").Open(EMessageSeverity::Warning);
+				;
+			FMessageLog("MapCheck").Open(EMessageSeverity::Warning);
+		}
 	}
 
 #if WITH_EDITOR
