Foliage Actor: Fix Initialize Flag not being valid on load

#rnx

diff --git a/Engine/Source/Runtime/Foliage/Private/FoliageActor.cpp b/Engine/Source/Runtime/Foliage/Private/FoliageActor.cpp
--- a/Engine/Source/Runtime/Foliage/Private/FoliageActor.cpp
+++ b/Engine/Source/Runtime/Foliage/Private/FoliageActor.cpp
@@ -44,7 +44,7 @@ void FFoliageActor::DestroyActors(bool bOnLoad)
 #if WITH_EDITOR
 bool FFoliageActor::IsInitialized() const
 {
-	return ActorClass != nullptr;;
+	return ActorClass != nullptr;
 }
 
 void FFoliageActor::Initialize(AInstancedFoliageActor* IFA, const UFoliageType* FoliageType)
