Mac: Fix conditional expression is ambiguous

#rnx

diff --git a/Engine/Source/Runtime/Foliage/Private/FoliageActor.cpp b/Engine/Source/Runtime/Foliage/Private/FoliageActor.cpp
--- a/Engine/Source/Runtime/Foliage/Private/FoliageActor.cpp
+++ b/Engine/Source/Runtime/Foliage/Private/FoliageActor.cpp
@@ -51,7 +51,7 @@ void FFoliageActor::Initialize(AInstancedFoliageActor* IFA, const UFoliageType*
 {
 	check(!IsInitialized());
 	const UFoliageType_Actor* FoliageType_Actor = Cast<UFoliageType_Actor>(FoliageType);
-	ActorClass = FoliageType_Actor->ActorClass ? FoliageType_Actor->ActorClass : AActor::StaticClass();
+	ActorClass = FoliageType_Actor->ActorClass ? *FoliageType_Actor->ActorClass : AActor::StaticClass();
 }
 
 void FFoliageActor::Uninitialize()
