Duplicate Fix (CL 7099700): Foliage Actor Fixes

#rnx

diff --git a/Engine/Source/Editor/UnrealEd/Private/Kismet2/KismetReinstanceUtilities.cpp b/Engine/Source/Editor/UnrealEd/Private/Kismet2/KismetReinstanceUtilities.cpp
--- a/Engine/Source/Editor/UnrealEd/Private/Kismet2/KismetReinstanceUtilities.cpp
+++ b/Engine/Source/Editor/UnrealEd/Private/Kismet2/KismetReinstanceUtilities.cpp
@@ -1790,6 +1790,13 @@ static void ReplaceActorHelper(UObject* OldObject, UClass* OldClass, UObject*& N
 	SpawnInfo.bDeferConstruction = true;
 	SpawnInfo.Name = OldActor->GetFName();
 
+#if WITH_EDITOR
+	if (!OldActor->IsListedInSceneOutliner())
+	{
+		SpawnInfo.bHideFromSceneOutliner = true;
+	}
+#endif
+
 	OldActor->UObject::Rename(nullptr, OldObject->GetOuter(), REN_DoNotDirty | REN_DontCreateRedirectors | (GBlueprintUseCompilationManager ? REN_ForceNoResetLoaders : 0));
 
 	AActor* NewActor = nullptr;
diff --git a/Engine/Source/Runtime/Foliage/Private/FoliageActor.cpp b/Engine/Source/Runtime/Foliage/Private/FoliageActor.cpp
--- a/Engine/Source/Runtime/Foliage/Private/FoliageActor.cpp
+++ b/Engine/Source/Runtime/Foliage/Private/FoliageActor.cpp
@@ -11,9 +11,9 @@ FName FoliageActorTag(TEXT("FoliageActorInstance"));
 // FFoliageActor
 void FFoliageActor::AddReferencedObjects(UObject* InThis, FReferenceCollector& Collector)
 {
-	for (const TWeakObjectPtr<AActor>& ActorPtr : ActorInstances)
+	for (AActor* Actor : ActorInstances)
 	{
-		if (AActor* Actor = ActorPtr.Get(false))
+		if (Actor != nullptr)
 		{
 			Collector.AddReferencedObject(Actor, InThis);
 		}
@@ -22,23 +22,35 @@ void FFoliageActor::AddReferencedObjects(UObject* InThis, FReferenceCollector& C
 
 void FFoliageActor::Serialize(FArchive& Ar)
 {
-	Ar << ActorInstances;
+#if WITH_EDITORONLY_DATA
+	if (Ar.CustomVer(FFoliageCustomVersion::GUID) < FFoliageCustomVersion::FoliageActorSupportNoWeakPtr)
+	{
+		Ar << ActorInstances_Deprecated;
+	}
+	else
+#endif
+	{
+		Ar << ActorInstances;
+	}
+
+		
 	Ar << ActorClass;
 }
 
 void FFoliageActor::DestroyActors(bool bOnLoad)
 {
-	for (const TWeakObjectPtr<AActor>& ActorPtr : ActorInstances)
+	for (AActor* Actor : ActorInstances)
 	{
-		if (AActor* Actor = ActorPtr.Get(false))
+		if (Actor != nullptr)
 		{
-			if (bOnLoad)
+            if (bOnLoad)
 			{
 				Actor->ConditionalPostLoad();
 			}
 			Actor->GetWorld()->DestroyActor(Actor);
 		}
 	}
+	ActorInstances.Empty();
 }
 
 #if WITH_EDITOR
@@ -58,7 +70,6 @@ void FFoliageActor::Uninitialize()
 {
 	check(IsInitialized());
 	DestroyActors(false);
-	ActorInstances.Empty();
 	ActorClass = nullptr;
 }
 
@@ -88,9 +99,9 @@ TArray<AActor*> FFoliageActor::GetActorsFromSelectedIndices(const TSet<int32>& S
 	for (int32 i : SelectedIndices)
 	{
 		check(i < ActorInstances.Num());
-		if (ActorInstances[i].IsValid())
+		if (ActorInstances[i] != nullptr)
 		{
-			Selection.Add(ActorInstances[i].Get());
+			Selection.Add(ActorInstances[i]);
 		}
 	}
 
@@ -118,7 +129,7 @@ void FFoliageActor::AddInstance(AInstancedFoliageActor* IFA, const FFoliageInsta
 
 void FFoliageActor::RemoveInstance(int32 InstanceIndex)
 {
-	AActor* Actor = ActorInstances[InstanceIndex].Get(false);
+	AActor* Actor = ActorInstances[InstanceIndex];
 	ActorInstances.RemoveAtSwap(InstanceIndex);
 
 	if (Actor && Actor->GetWorld())
@@ -129,7 +140,7 @@ void FFoliageActor::RemoveInstance(int32 InstanceIndex)
 
 void FFoliageActor::SetInstanceWorldTransform(int32 InstanceIndex, const FTransform& Transform, bool bTeleport)
 {
-	if (AActor* Actor = ActorInstances[InstanceIndex].Get(false))
+	if (AActor* Actor = ActorInstances[InstanceIndex])
 	{
 		Actor->SetActorTransform(Transform);
 	}
@@ -137,7 +148,7 @@ void FFoliageActor::SetInstanceWorldTransform(int32 InstanceIndex, const FTransf
 
 FTransform FFoliageActor::GetInstanceWorldTransform(int32 InstanceIndex) const
 {
-	if (AActor* Actor = ActorInstances[InstanceIndex].Get(false))
+	if (AActor* Actor = ActorInstances[InstanceIndex])
 	{
 		return Actor->GetTransform();
 	}
@@ -175,9 +186,9 @@ void FFoliageActor::Refresh(AInstancedFoliageActor* IFA, const TArray<FFoliageIn
 
 void FFoliageActor::OnHiddenEditorViewMaskChanged(uint64 InHiddenEditorViews)
 {
-	for (const TWeakObjectPtr<AActor>& ActorPtr : ActorInstances)
+	for (AActor* Actor : ActorInstances)
 	{
-		if (AActor* Actor = ActorPtr.Get(false))
+		if (Actor != nullptr)
 		{
 			if (Actor->HiddenEditorViews != InHiddenEditorViews)
 			{
@@ -209,8 +220,7 @@ void FFoliageActor::Reapply(AInstancedFoliageActor* IFA, const UFoliageType* Fol
 {
 	IFA->Modify();
 	DestroyActors(false);
-	ActorInstances.Empty();
-
+	
 	if (IsInitialized())
 	{
 		Uninitialize();
@@ -230,9 +240,9 @@ void FFoliageActor::SelectInstances(bool bSelect, int32 InstanceIndex, int32 Cou
 	check(ActorInstances.Num() >= InstanceIndex + Count);
 	for (int32 i = InstanceIndex; i < (InstanceIndex + Count); ++i)
 	{
-		if (ActorInstances[i].IsValid())
+		if (ActorInstances[i])
 		{
-			Selection.Add(ActorInstances[i].Get());
+			Selection.Add(ActorInstances[i]);
 		}
 	}
 	AInstancedFoliageActor::SelectionChanged.Broadcast(bSelect, Selection);
diff --git a/Engine/Source/Runtime/Foliage/Private/FoliageActor.h b/Engine/Source/Runtime/Foliage/Private/FoliageActor.h
--- a/Engine/Source/Runtime/Foliage/Private/FoliageActor.h
+++ b/Engine/Source/Runtime/Foliage/Private/FoliageActor.h
@@ -5,7 +5,11 @@
 
 struct FFoliageActor : public FFoliageImpl
 {
-	TArray<TWeakObjectPtr<AActor>> ActorInstances;
+#if WITH_EDITORONLY_DATA
+	TArray<TWeakObjectPtr<AActor>> ActorInstances_Deprecated;
+#endif
+
+	TArray<AActor*> ActorInstances;
 	UClass* ActorClass;
 	
 	FFoliageActor()
diff --git a/Engine/Source/Runtime/Foliage/Private/InstancedFoliage.cpp b/Engine/Source/Runtime/Foliage/Private/InstancedFoliage.cpp
--- a/Engine/Source/Runtime/Foliage/Private/InstancedFoliage.cpp
+++ b/Engine/Source/Runtime/Foliage/Private/InstancedFoliage.cpp
@@ -66,49 +66,6 @@ static TAutoConsoleVariable<int32> CVarFoliageDiscardDataOnLoad(
 	TEXT("1: Discard scalable foliage data on load (disables all scalable foliage types); 0: Keep scalable foliage data (requires reloading level)"),
 	ECVF_Scalability);
 
-// Custom serialization version for all packages containing Instance Foliage
-struct FFoliageCustomVersion
-{
-	enum Type
-	{
-		// Before any version changes were made in the plugin
-		BeforeCustomVersionWasAdded = 0,
-		// Converted to use HierarchicalInstancedStaticMeshComponent
-		FoliageUsingHierarchicalISMC = 1,
-		// Changed Component to not RF_Transactional
-		HierarchicalISMCNonTransactional = 2,
-		// Added FoliageTypeUpdateGuid
-		AddedFoliageTypeUpdateGuid = 3,
-		// Use a GUID to determine whic procedural actor spawned us
-		ProceduralGuid = 4,
-		// Support for cross-level bases 
-		CrossLevelBase = 5,
-		// FoliageType for details customization
-		FoliageTypeCustomization = 6,
-		// FoliageType for details customization continued
-		FoliageTypeCustomizationScaling = 7,
-		// FoliageType procedural scale and shade settings updated
-		FoliageTypeProceduralScaleAndShade = 8,
-		// Added FoliageHISMC and blueprint support
-		FoliageHISMCBlueprints = 9,
-		// Added Mobility setting to UFoliageType
-		AddedMobility = 10,
-		// Make sure that foliage has FoliageHISMC class
-		FoliageUsingFoliageISMC = 11,
-		// Foliage Actor Support
-		FoliageActorSupport = 12,
-		// -----<new versions can be added above this line>-------------------------------------------------
-		VersionPlusOne,
-		LatestVersion = VersionPlusOne - 1
-	};
-
-	// The GUID for this custom version number
-	const static FGuid GUID;
-
-private:
-	FFoliageCustomVersion() {}
-};
-
 const FGuid FFoliageCustomVersion::GUID(0x430C4D19, 0x71544970, 0x87699B69, 0xDF90B0E5);
 // Register the custom version with core
 FCustomVersionRegistration GRegisterFoliageCustomVersion(FFoliageCustomVersion::GUID, FFoliageCustomVersion::LatestVersion, TEXT("FoliageVer"));
@@ -3544,6 +3501,19 @@ void AInstancedFoliageActor::PostLoad()
 				}
 			}
 
+			if (GetLinkerCustomVersion(FFoliageCustomVersion::GUID) < FFoliageCustomVersion::FoliageActorSupportNoWeakPtr)
+			{
+				if (Info.Type == EFoliageImplType::Actor)
+				{
+					FFoliageActor* FoliageActor = StaticCast<FFoliageActor*>(Info.Implementation.Get());
+					for (const TWeakObjectPtr<AActor>& ActorPtr : FoliageActor->ActorInstances_Deprecated)
+					{
+						FoliageActor->ActorInstances.Add(ActorPtr.Get());
+					}
+					FoliageActor->ActorInstances_Deprecated.Empty();
+				}
+			}
+
 			// Update foliage component settings if the foliage settings object was changed while the level was not loaded.
 			if (Info.FoliageTypeUpdateGuid != FoliageType->UpdateGuid)
 			{
@@ -3641,22 +3611,25 @@ void AInstancedFoliageActor::PostLoad()
 	{
 		for (auto& Pair : FoliageInfos)
 		{
-			if (Pair.Value->Type == EFoliageImplType::StaticMesh)
+			if (!Pair.Key || Pair.Key->bEnableDensityScaling)
 			{
-				FFoliageStaticMesh* FoliageStaticMesh = StaticCast<FFoliageStaticMesh*>(Pair.Value->Implementation.Get());
-				
-				if (FoliageStaticMesh->Component != nullptr && (!Pair.Key || Pair.Key->bEnableDensityScaling))
+				if (Pair.Value->Type == EFoliageImplType::StaticMesh)
 				{
-					FoliageStaticMesh->Component->ConditionalPostLoad();
-					FoliageStaticMesh->Component->DestroyComponent();
+					FFoliageStaticMesh* FoliageStaticMesh = StaticCast<FFoliageStaticMesh*>(Pair.Value->Implementation.Get());
+
+					if (FoliageStaticMesh->Component != nullptr)
+					{
+						FoliageStaticMesh->Component->ConditionalPostLoad();
+						FoliageStaticMesh->Component->DestroyComponent();
+					}
+				}
+				else if (Pair.Value->Type == EFoliageImplType::Actor)
+				{
+					FFoliageActor* FoliageActor = StaticCast<FFoliageActor*>(Pair.Value->Implementation.Get());
+					FoliageActor->DestroyActors(true);
 				}
-			}	
-			else if (Pair.Value->Type == EFoliageImplType::Actor)
-			{
-				FFoliageActor* FoliageActor = StaticCast<FFoliageActor*>(Pair.Value->Implementation.Get());
-				FoliageActor->DestroyActors(true);
 			}
-
+				
 			Pair.Value = FFoliageInfo();
 		}
 	}
diff --git a/Engine/Source/Runtime/Foliage/Public/InstancedFoliageActor.h b/Engine/Source/Runtime/Foliage/Public/InstancedFoliageActor.h
--- a/Engine/Source/Runtime/Foliage/Public/InstancedFoliageActor.h
+++ b/Engine/Source/Runtime/Foliage/Public/InstancedFoliageActor.h
@@ -13,6 +13,51 @@
 
 #include "InstancedFoliageActor.generated.h"
 
+// Custom serialization version for all packages containing Instance Foliage
+struct FFoliageCustomVersion
+{
+	enum Type
+	{
+		// Before any version changes were made in the plugin
+		BeforeCustomVersionWasAdded = 0,
+		// Converted to use HierarchicalInstancedStaticMeshComponent
+		FoliageUsingHierarchicalISMC = 1,
+		// Changed Component to not RF_Transactional
+		HierarchicalISMCNonTransactional = 2,
+		// Added FoliageTypeUpdateGuid
+		AddedFoliageTypeUpdateGuid = 3,
+		// Use a GUID to determine whic procedural actor spawned us
+		ProceduralGuid = 4,
+		// Support for cross-level bases 
+		CrossLevelBase = 5,
+		// FoliageType for details customization
+		FoliageTypeCustomization = 6,
+		// FoliageType for details customization continued
+		FoliageTypeCustomizationScaling = 7,
+		// FoliageType procedural scale and shade settings updated
+		FoliageTypeProceduralScaleAndShade = 8,
+		// Added FoliageHISMC and blueprint support
+		FoliageHISMCBlueprints = 9,
+		// Added Mobility setting to UFoliageType
+		AddedMobility = 10,
+		// Make sure that foliage has FoliageHISMC class
+		FoliageUsingFoliageISMC = 11,
+		// Foliage Actor Support
+		FoliageActorSupport = 12,
+		// Foliage Actor (No weak ptr)
+		FoliageActorSupportNoWeakPtr = 13,
+		// -----<new versions can be added above this line>-------------------------------------------------
+		VersionPlusOne,
+		LatestVersion = VersionPlusOne - 1
+	};
+
+	// The GUID for this custom version number
+	const static FGuid GUID;
+
+private:
+	FFoliageCustomVersion() {}
+};
+
 class UProceduralFoliageComponent;
 
 // Function for filtering out hit components during FoliageTrace
