Integrate 7321522 - Foliage Actor Fix: AddReferencedObjects was passing down local ref.

#rnx

diff --git a/Engine/Source/Runtime/Foliage/Private/FoliageActor.cpp b/Engine/Source/Runtime/Foliage/Private/FoliageActor.cpp
--- a/Engine/Source/Runtime/Foliage/Private/FoliageActor.cpp
+++ b/Engine/Source/Runtime/Foliage/Private/FoliageActor.cpp
@@ -11,7 +11,7 @@ FName FoliageActorTag(TEXT("FoliageActorInstance"));
 // FFoliageActor
 void FFoliageActor::AddReferencedObjects(UObject* InThis, FReferenceCollector& Collector)
 {
-	for (AActor* Actor : ActorInstances)
+	for (AActor*& Actor : ActorInstances)
 	{
 		if (Actor != nullptr)
 		{
