No need for AddReferencedObjects in FFoliageActor since it is pointing to Level Actors.

#rnx

diff --git a/Engine/Source/Runtime/Foliage/Private/FoliageActor.cpp b/Engine/Source/Runtime/Foliage/Private/FoliageActor.cpp
--- a/Engine/Source/Runtime/Foliage/Private/FoliageActor.cpp
+++ b/Engine/Source/Runtime/Foliage/Private/FoliageActor.cpp
@@ -9,16 +9,6 @@ FName FoliageActorTag(TEXT("FoliageActorInstance"));
 //
 //
 // FFoliageActor
-void FFoliageActor::AddReferencedObjects(UObject* InThis, FReferenceCollector& Collector)
-{
-	for (AActor*& Actor : ActorInstances)
-	{
-		if (Actor != nullptr)
-		{
-			Collector.AddReferencedObject(Actor, InThis);
-		}
-	}
-}
 
 void FFoliageActor::Serialize(FArchive& Ar)
 {
diff --git a/Engine/Source/Runtime/Foliage/Private/FoliageActor.h b/Engine/Source/Runtime/Foliage/Private/FoliageActor.h
--- a/Engine/Source/Runtime/Foliage/Private/FoliageActor.h
+++ b/Engine/Source/Runtime/Foliage/Private/FoliageActor.h
@@ -17,7 +17,6 @@ struct FFoliageActor : public FFoliageImpl
 	{
 	}
 
-	virtual void AddReferencedObjects(UObject* InThis, FReferenceCollector& Collector) override;
 	virtual void Serialize(FArchive& Ar) override;
 	void DestroyActors(bool bOnLoad);
 
