Added accessors to empty a procedural foliage spawner.

diff --git a/Engine/Source/Runtime/Foliage/Private/ProceduralFoliageComponent.cpp b/Engine/Source/Runtime/Foliage/Private/ProceduralFoliageComponent.cpp
--- a/Engine/Source/Runtime/Foliage/Private/ProceduralFoliageComponent.cpp
+++ b/Engine/Source/Runtime/Foliage/Private/ProceduralFoliageComponent.cpp
@@ -192,7 +192,7 @@ bool UProceduralFoliageComponent::ExecuteSimulation(TArray<FDesiredFoliageInstan
 					const FTransform TileTM(OrientedOffset + WorldPosition);
 
 					TArray<FDesiredFoliageInstance>* DesiredInstances = new TArray<FDesiredFoliageInstance>();
-					CompositeTile->ExtractDesiredInstances(*DesiredInstances, TileTM, ProceduralGuid, TileLayout.HalfHeight, BoundsBodyInstance);
+					CompositeTile->ExtractDesiredInstances(*DesiredInstances, TileTM, ProceduralGuid, TileLayout.HalfHeight, BoundsBodyInstance, true);
 
 					return DesiredInstances;
 					
diff --git a/Engine/Source/Runtime/Foliage/Private/ProceduralFoliageSpawner.cpp b/Engine/Source/Runtime/Foliage/Private/ProceduralFoliageSpawner.cpp
--- a/Engine/Source/Runtime/Foliage/Private/ProceduralFoliageSpawner.cpp
+++ b/Engine/Source/Runtime/Foliage/Private/ProceduralFoliageSpawner.cpp
@@ -111,6 +111,19 @@ void UProceduralFoliageSpawner::SimulateIfNeeded()
 
 }
 
+void UProceduralFoliageSpawner::Empty()
+{
+	for (TWeakObjectPtr<UProceduralFoliageTile>& WeakTile : PrecomputedTiles)
+	{
+		if (UProceduralFoliageTile* Tile = WeakTile.Get())
+		{
+			Tile->Empty();
+		}
+	}
+
+	PrecomputedTiles.Empty();
+}
+
 const UProceduralFoliageTile* UProceduralFoliageSpawner::GetRandomTile(int32 X, int32 Y)
 {
 	if (PrecomputedTiles.Num())
diff --git a/Engine/Source/Runtime/Foliage/Public/ProceduralFoliageSpawner.h b/Engine/Source/Runtime/Foliage/Public/ProceduralFoliageSpawner.h
--- a/Engine/Source/Runtime/Foliage/Public/ProceduralFoliageSpawner.h
+++ b/Engine/Source/Runtime/Foliage/Public/ProceduralFoliageSpawner.h
@@ -47,6 +47,7 @@ private:
 public:
 	UFUNCTION(BlueprintCallable, Category = ProceduralFoliageSimulation)
 	void Simulate(int32 NumSteps = -1);
+	void Empty();
 
 	int32 GetRandomNumber();
 
