Fixed issue where foliage meshes were not being moved or deleted when with the geometry they are placed on is moved or deleted.

diff --git a/Engine/Source/Runtime/Foliage/Private/InstancedFoliage.cpp b/Engine/Source/Runtime/Foliage/Private/InstancedFoliage.cpp
--- a/Engine/Source/Runtime/Foliage/Private/InstancedFoliage.cpp
+++ b/Engine/Source/Runtime/Foliage/Private/InstancedFoliage.cpp
@@ -3038,7 +3038,7 @@ void AInstancedFoliageActor::OnLevelActorMoved(AActor* InActor)
 
 	if (!InWorld || !InWorld->IsGameWorld())
 	{
-		for (UActorComponent* Component : GetComponents())
+		for (UActorComponent* Component : InActor->GetComponents())
 		{
 			if (Component)
 			{
@@ -3054,7 +3054,7 @@ void AInstancedFoliageActor::OnLevelActorDeleted(AActor* InActor)
 
 	if (!InWorld || !InWorld->IsGameWorld())
 	{
-		for (UActorComponent* Component : GetComponents())
+		for (UActorComponent* Component : InActor->GetComponents())
 		{
 			if (Component)
 			{
