Prevent possible crash if InstanceToMoves contained invalid entries

diff --git a/Engine/Source/Runtime/Foliage/Private/InstancedFoliage.cpp b/Engine/Source/Runtime/Foliage/Private/InstancedFoliage.cpp
--- a/Engine/Source/Runtime/Foliage/Private/InstancedFoliage.cpp
+++ b/Engine/Source/Runtime/Foliage/Private/InstancedFoliage.cpp
@@ -1925,6 +1925,7 @@ void AInstancedFoliageActor::MoveInstancesToNewComponent(UPrimitiveComponent* In
 
 	for (auto& MeshPair : FoliageMeshes)
 	{
+		InstancesToMove.Reset();
 		FFoliageMeshInfo& MeshInfo = *MeshPair.Value;
 		
 		if (MeshInfo.Component != nullptr)
@@ -1938,8 +1939,11 @@ void AInstancedFoliageActor::MoveInstancesToNewComponent(UPrimitiveComponent* In
 		// Add the foliage to the new level
 		for (int32 InstanceIndex : InstancesToMove)
 		{
-			FFoliageInstance NewInstance = MeshInfo.Instances[InstanceIndex];
-			TargetMeshInfo->AddInstance(TargetIFA, TargetFoliageType, NewInstance, InNewComponent, false);
+			if (MeshInfo.Instances.IsValidIndex(InstanceIndex))
+			{
+				FFoliageInstance NewInstance = MeshInfo.Instances[InstanceIndex];
+				TargetMeshInfo->AddInstance(TargetIFA, TargetFoliageType, NewInstance, InNewComponent, false);
+			}
 		}
 
 		if (TargetMeshInfo->Component != nullptr)
