Fix a iOS crash when CVarFoliageDiscardDataOnLoad=1

#rnx

diff --git a/Engine/Source/Runtime/Foliage/Private/InstancedFoliage.cpp b/Engine/Source/Runtime/Foliage/Private/InstancedFoliage.cpp
--- a/Engine/Source/Runtime/Foliage/Private/InstancedFoliage.cpp
+++ b/Engine/Source/Runtime/Foliage/Private/InstancedFoliage.cpp
@@ -389,8 +389,11 @@ FArchive& operator<<(FArchive& Ar, FFoliageInfo& Info)
 	{
 		Info.CreateImplementation(Info.Type);
 	}
-		
-	Info.Implementation->Serialize(Ar);
+	
+	if (Info.Implementation)
+	{
+		Info.Implementation->Serialize(Ar);
+	}
 
 #if WITH_EDITORONLY_DATA
 	if (!Ar.ArIsFilterEditorOnly && !(Ar.GetPortFlags() & PPF_DuplicateForPIE))
