Cherry pick CL 7891297: Fix Foliage: Foliage Actors need to reapply Foliage instance positions after ApplyLevelTransform is called on the Level

#rnx

diff --git a/Engine/Source/Runtime/Foliage/Private/FoliageActor.cpp b/Engine/Source/Runtime/Foliage/Private/FoliageActor.cpp
--- a/Engine/Source/Runtime/Foliage/Private/FoliageActor.cpp
+++ b/Engine/Source/Runtime/Foliage/Private/FoliageActor.cpp
@@ -199,7 +199,7 @@ void FFoliageActor::OnHiddenEditorViewMaskChanged(uint64 InHiddenEditorViews)
 	}
 }
 
-void FFoliageActor::PostEditUndo(AInstancedFoliageActor* IFA, UFoliageType* FoliageType, const TArray<FFoliageInstance>& Instances, const TSet<int32>& SelectedIndices)
+void FFoliageActor::UpdateActorTransforms(const TArray<FFoliageInstance>& Instances)
 {
 	for (int32 i = 0; i < Instances.Num(); ++i)
 	{
@@ -207,6 +207,16 @@ void FFoliageActor::PostEditUndo(AInstancedFoliageActor* IFA, UFoliageType* Foli
 	}
 }
 
+void FFoliageActor::PostApplyLevelTransform(const FTransform& InTransform, const TArray<FFoliageInstance>& Instances)
+{
+	UpdateActorTransforms(Instances);
+}
+
+void FFoliageActor::PostEditUndo(AInstancedFoliageActor* IFA, UFoliageType* FoliageType, const TArray<FFoliageInstance>& Instances, const TSet<int32>& SelectedIndices)
+{
+	UpdateActorTransforms(Instances);
+}
+
 void FFoliageActor::NotifyFoliageTypeChanged(AInstancedFoliageActor* IFA, UFoliageType* FoliageType, const TArray<FFoliageInstance>& Instances, const TSet<int32>& SelectedIndices, bool bSourceChanged)
 {
 	if (bSourceChanged)
@@ -216,10 +226,10 @@ void FFoliageActor::NotifyFoliageTypeChanged(AInstancedFoliageActor* IFA, UFolia
 	}
 }
 
-void FFoliageActor::Reapply(AInstancedFoliageActor* IFA, const UFoliageType* FoliageType, const TArray<FFoliageInstance>& Instances)
+void FFoliageActor::Reapply(AInstancedFoliageActor* IFA, const UFoliageType* FoliageType, const TArray<FFoliageInstance>& Instances, bool bPostLoad)
 {
 	IFA->Modify();
-	DestroyActors(false);
+	DestroyActors(bPostLoad);
 	
 	if (IsInitialized())
 	{
diff --git a/Engine/Source/Runtime/Foliage/Private/FoliageActor.h b/Engine/Source/Runtime/Foliage/Private/FoliageActor.h
--- a/Engine/Source/Runtime/Foliage/Private/FoliageActor.h
+++ b/Engine/Source/Runtime/Foliage/Private/FoliageActor.h
@@ -41,9 +41,13 @@ struct FFoliageActor : public FFoliageImpl
 	virtual void Refresh(AInstancedFoliageActor* IFA, const TArray<FFoliageInstance>& Instances, bool Async, bool Force) override;
 	virtual void OnHiddenEditorViewMaskChanged(uint64 InHiddenEditorViews) override;
 	virtual void PostEditUndo(AInstancedFoliageActor* IFA, UFoliageType* FoliageType, const TArray<FFoliageInstance>& Instances, const TSet<int32>& SelectedIndices) override;
+	virtual void PostApplyLevelTransform(const FTransform& InTransform, const TArray<FFoliageInstance>& Instances) override;
 	virtual void NotifyFoliageTypeChanged(AInstancedFoliageActor* IFA, UFoliageType* FoliageType, const TArray<FFoliageInstance>& Instances, const TSet<int32>& SelectedIndices, bool bSourceChanged) override;
-	void Reapply(AInstancedFoliageActor* IFA, const UFoliageType* FoliageType, const TArray<FFoliageInstance>& Instances);
+	void Reapply(AInstancedFoliageActor* IFA, const UFoliageType* FoliageType, const TArray<FFoliageInstance>& Instances, bool bPostLoad = false);
 	AActor* Spawn(AInstancedFoliageActor* IFA, const FFoliageInstance& Instance);
 	TArray<AActor*> GetActorsFromSelectedIndices(const TSet<int32>& SelectedIndices) const;
+
+private:
+	void UpdateActorTransforms(const TArray<FFoliageInstance>& Instances);
 #endif
 };
\ No newline at end of file
diff --git a/Engine/Source/Runtime/Foliage/Private/InstancedFoliage.cpp b/Engine/Source/Runtime/Foliage/Private/InstancedFoliage.cpp
--- a/Engine/Source/Runtime/Foliage/Private/InstancedFoliage.cpp
+++ b/Engine/Source/Runtime/Foliage/Private/InstancedFoliage.cpp
@@ -3347,6 +3347,11 @@ void AInstancedFoliageActor::PostInitProperties()
 		GEngine->OnLevelActorDeleted().Remove(OnLevelActorDeletedDelegateHandle);
 		OnLevelActorDeletedDelegateHandle = GEngine->OnLevelActorDeleted().AddUObject(this, &AInstancedFoliageActor::OnLevelActorDeleted);
 
+		if (GetLevel())
+		{
+			OnApplyLevelTransformDelegateHandle = GetLevel()->OnApplyLevelTransform.AddUObject(this, &AInstancedFoliageActor::OnApplyLevelTransform);
+		}
+
 		FWorldDelegates::PostApplyLevelOffset.Remove(OnPostApplyLevelOffsetDelegateHandle);
 		OnPostApplyLevelOffsetDelegateHandle = FWorldDelegates::PostApplyLevelOffset.AddUObject(this, &AInstancedFoliageActor::OnPostApplyLevelOffset);
 	}
@@ -3360,6 +3365,12 @@ void AInstancedFoliageActor::BeginDestroy()
 	{
 		GEngine->OnActorMoved().Remove(OnLevelActorMovedDelegateHandle);
 		GEngine->OnLevelActorDeleted().Remove(OnLevelActorDeletedDelegateHandle);
+		
+		if (GetLevel())
+		{
+			GetLevel()->OnApplyLevelTransform.Remove(OnApplyLevelTransformDelegateHandle);
+		}
+
 		FWorldDelegates::PostApplyLevelOffset.Remove(OnPostApplyLevelOffsetDelegateHandle);
 	}
 }
@@ -3534,7 +3545,8 @@ void AInstancedFoliageActor::PostLoad()
 					else if (Info.Type == EFoliageImplType::Actor)
 					{
 						FFoliageActor* FoliageActor = StaticCast<FFoliageActor*>(Info.Implementation.Get());
-						FoliageActor->Reapply(this, FoliageType, Info.Instances);
+						const bool bPostLoad = true;
+						FoliageActor->Reapply(this, FoliageType, Info.Instances, bPostLoad);
 					}
 				}
 				Info.FoliageTypeUpdateGuid = FoliageType->UpdateGuid;
@@ -3732,6 +3744,18 @@ void AInstancedFoliageActor::OnLevelActorDeleted(AActor* InActor)
 	}
 }
 
+void AInstancedFoliageActor::OnApplyLevelTransform(const FTransform& InTransform)
+{
+	for (auto& Pair : FoliageInfos)
+	{
+		FFoliageInfo& Info = *Pair.Value;
+		if (Info.Implementation)
+		{
+			Info.Implementation->PostApplyLevelTransform(InTransform, Info.Instances);
+		}
+	}
+}
+
 void AInstancedFoliageActor::OnPostApplyLevelOffset(ULevel* InLevel, UWorld* InWorld, const FVector& InOffset, bool bWorldShift)
 {
 	ULevel* OwningLevel = GetLevel();
diff --git a/Engine/Source/Runtime/Foliage/Public/InstancedFoliage.h b/Engine/Source/Runtime/Foliage/Public/InstancedFoliage.h
--- a/Engine/Source/Runtime/Foliage/Public/InstancedFoliage.h
+++ b/Engine/Source/Runtime/Foliage/Public/InstancedFoliage.h
@@ -202,6 +202,7 @@ struct FFoliageImpl
 	virtual void OnHiddenEditorViewMaskChanged(uint64 InHiddenEditorViews) = 0;
 	virtual void PreEditUndo(AInstancedFoliageActor* IFA, UFoliageType* FoliageType) {}
 	virtual void PostEditUndo(AInstancedFoliageActor* IFA, UFoliageType* FoliageType, const TArray<FFoliageInstance>& Instances, const TSet<int32>& SelectedIndices) = 0;
+	virtual void PostApplyLevelTransform(const FTransform& InTransform, const TArray<FFoliageInstance>& Instances) {}
 	virtual void NotifyFoliageTypeWillChange(AInstancedFoliageActor* IFA, UFoliageType* FoliageType) {}
 	virtual void NotifyFoliageTypeChanged(AInstancedFoliageActor* IFA, UFoliageType* FoliageType, const TArray<FFoliageInstance>& Instances, const TSet<int32>& SelectedIndices, bool bSourceChanged) = 0;
 	virtual void EnterEditMode() {}
diff --git a/Engine/Source/Runtime/Foliage/Public/InstancedFoliageActor.h b/Engine/Source/Runtime/Foliage/Public/InstancedFoliageActor.h
--- a/Engine/Source/Runtime/Foliage/Public/InstancedFoliageActor.h
+++ b/Engine/Source/Runtime/Foliage/Public/InstancedFoliageActor.h
@@ -251,6 +251,7 @@ private:
 	void ClearSelection();
 	void OnLevelActorMoved(AActor* InActor);
 	void OnLevelActorDeleted(AActor* InActor);
+	void OnApplyLevelTransform(const FTransform& InTransform);
 	void OnPostApplyLevelOffset(ULevel* InLevel, UWorld* InWorld, const FVector& InOffset, bool bWorldShift);
 
 	// Move instances to a foliage actor in target level
@@ -261,6 +262,7 @@ private:
 	FDelegateHandle OnLevelActorMovedDelegateHandle;
 	FDelegateHandle OnLevelActorDeletedDelegateHandle;
 	FDelegateHandle OnPostApplyLevelOffsetDelegateHandle;
+	FDelegateHandle OnApplyLevelTransformDelegateHandle;
 
 	FOnFoliageTypeMeshChanged OnFoliageTypeMeshChangedEvent;
 #endif
