Fixed Foliage instances not properly handling Level transform

diff --git a/Engine/Source/Runtime/Engine/Private/LevelUtils.cpp b/Engine/Source/Runtime/Engine/Private/LevelUtils.cpp
--- a/Engine/Source/Runtime/Engine/Private/LevelUtils.cpp
+++ b/Engine/Source/Runtime/Engine/Private/LevelUtils.cpp
@@ -46,6 +46,7 @@ static TMap<ULevel*, FLevelReadOnlyData> LevelReadOnlyCache;
 
 #if WITH_EDITOR
 bool FLevelUtils::bMovingLevel = false;
+bool FLevelUtils::bApplyingLevelTransform = false;
 #endif
 
 /**
@@ -371,14 +372,22 @@ bool FLevelUtils::IsMovingLevel()
 	return bMovingLevel;
 }
 
-#endif // WITH_EDITOR
+bool FLevelUtils::IsApplyingLevelTransform()
+{
+	return bApplyingLevelTransform;
+}
 
+#endif // WITH_EDITOR
 
 void FLevelUtils::ApplyLevelTransform( ULevel* Level, const FTransform& LevelTransform, bool bDoPostEditMove )
 {
 	bool bTransformActors =  !LevelTransform.Equals(FTransform::Identity);
 	if (bTransformActors)
 	{
+#if WITH_EDITOR
+		bool bOldValue = bApplyingLevelTransform;
+		bApplyingLevelTransform = true;
+#endif
 		if (!LevelTransform.GetRotation().IsIdentity())
 		{
 			// If there is a rotation applied, then the relative precomputed bounds become invalid.
@@ -410,6 +419,9 @@ void FLevelUtils::ApplyLevelTransform( ULevel* Level, const FTransform& LevelTra
 #endif // WITH_EDITOR
 
 		Level->OnApplyLevelTransform.Broadcast(LevelTransform);
+#if WITH_EDITOR
+		bApplyingLevelTransform = bOldValue;
+#endif
 	}
 }
 
diff --git a/Engine/Source/Runtime/Engine/Public/LevelUtils.h b/Engine/Source/Runtime/Engine/Public/LevelUtils.h
--- a/Engine/Source/Runtime/Engine/Public/LevelUtils.h
+++ b/Engine/Source/Runtime/Engine/Public/LevelUtils.h
@@ -140,11 +140,13 @@ public:
 	* Returns true if we are moving a level
 	*/
 	static bool IsMovingLevel();
+	static bool IsApplyingLevelTransform();
 
 private:
 
 	// Flag to mark if we are currently finalizing a level offset
 	static bool bMovingLevel;
+	static bool bApplyingLevelTransform;
 
 #endif // WITH_EDITOR
 };
diff --git a/Engine/Source/Runtime/Foliage/Private/FoliageActor.cpp b/Engine/Source/Runtime/Foliage/Private/FoliageActor.cpp
--- a/Engine/Source/Runtime/Foliage/Private/FoliageActor.cpp
+++ b/Engine/Source/Runtime/Foliage/Private/FoliageActor.cpp
@@ -207,11 +207,6 @@ void FFoliageActor::UpdateActorTransforms(const TArray<FFoliageInstance>& Instan
 	}
 }
 
-void FFoliageActor::PostApplyLevelTransform(const FTransform& InTransform, const TArray<FFoliageInstance>& Instances)
-{
-	UpdateActorTransforms(Instances);
-}
-
 void FFoliageActor::PostEditUndo(AInstancedFoliageActor* IFA, UFoliageType* FoliageType, const TArray<FFoliageInstance>& Instances, const TSet<int32>& SelectedIndices)
 {
 	UpdateActorTransforms(Instances);
diff --git a/Engine/Source/Runtime/Foliage/Private/FoliageActor.h b/Engine/Source/Runtime/Foliage/Private/FoliageActor.h
--- a/Engine/Source/Runtime/Foliage/Private/FoliageActor.h
+++ b/Engine/Source/Runtime/Foliage/Private/FoliageActor.h
@@ -41,7 +41,6 @@ struct FFoliageActor : public FFoliageImpl
 	virtual void Refresh(AInstancedFoliageActor* IFA, const TArray<FFoliageInstance>& Instances, bool Async, bool Force) override;
 	virtual void OnHiddenEditorViewMaskChanged(uint64 InHiddenEditorViews) override;
 	virtual void PostEditUndo(AInstancedFoliageActor* IFA, UFoliageType* FoliageType, const TArray<FFoliageInstance>& Instances, const TSet<int32>& SelectedIndices) override;
-	virtual void PostApplyLevelTransform(const FTransform& InTransform, const TArray<FFoliageInstance>& Instances) override;
 	virtual void NotifyFoliageTypeChanged(AInstancedFoliageActor* IFA, UFoliageType* FoliageType, const TArray<FFoliageInstance>& Instances, const TSet<int32>& SelectedIndices, bool bSourceChanged) override;
 	void Reapply(AInstancedFoliageActor* IFA, const UFoliageType* FoliageType, const TArray<FFoliageInstance>& Instances, bool bPostLoad = false);
 	AActor* Spawn(AInstancedFoliageActor* IFA, const FFoliageInstance& Instance);
diff --git a/Engine/Source/Runtime/Foliage/Private/InstancedFoliage.cpp b/Engine/Source/Runtime/Foliage/Private/InstancedFoliage.cpp
--- a/Engine/Source/Runtime/Foliage/Private/InstancedFoliage.cpp
+++ b/Engine/Source/Runtime/Foliage/Private/InstancedFoliage.cpp
@@ -45,6 +45,7 @@ InstancedFoliage.cpp: Instanced foliage implementation.
 #include "UObject/FortniteMainBranchObjectVersion.h"
 #include "PreviewScene.h"
 #include "FoliageActor.h"
+#include "LevelUtils.h"
 
 #define LOCTEXT_NAMESPACE "InstancedFoliage"
 
@@ -2315,6 +2316,11 @@ const FFoliageInfo* AInstancedFoliageActor::FindInfo(const UFoliageType* InType)
 #if WITH_EDITOR
 void AInstancedFoliageActor::MoveInstancesForMovedComponent(UActorComponent* InComponent)
 {
+	if (FLevelUtils::IsApplyingLevelTransform())
+	{
+		return;
+	}
+
 	const auto BaseId = InstanceBaseCache.GetInstanceBaseId(InComponent);
 	if (BaseId == FFoliageInstanceBaseCache::InvalidBaseId)
 	{
@@ -3746,14 +3752,26 @@ void AInstancedFoliageActor::OnLevelActorDeleted(AActor* InActor)
 
 void AInstancedFoliageActor::OnApplyLevelTransform(const FTransform& InTransform)
 {
+#if WITH_EDITORONLY_DATA
 	for (auto& Pair : FoliageInfos)
 	{
 		FFoliageInfo& Info = *Pair.Value;
-		if (Info.Implementation)
+
+		InstanceBaseCache.UpdateInstanceBaseCachedTransforms();
+
+		Info.InstanceHash->Empty();
+		for (int32 InstanceIdx = 0; InstanceIdx < Info.Instances.Num(); InstanceIdx++)
 		{
-			Info.Implementation->PostApplyLevelTransform(InTransform, Info.Instances);
+			FFoliageInstance& Instance = Info.Instances[InstanceIdx];
+			FTransform OldTransform(Instance.Rotation, Instance.Location);
+			FTransform NewTransform = OldTransform * InTransform;
+			Instance.Location = NewTransform.GetTranslation();
+			Instance.Rotation = NewTransform.Rotator();
+			// Rehash instance location
+			Info.InstanceHash->InsertInstance(Instance.Location, InstanceIdx);
 		}
 	}
+#endif
 }
 
 void AInstancedFoliageActor::OnPostApplyLevelOffset(ULevel* InLevel, UWorld* InWorld, const FVector& InOffset, bool bWorldShift)
diff --git a/Engine/Source/Runtime/Foliage/Public/InstancedFoliage.h b/Engine/Source/Runtime/Foliage/Public/InstancedFoliage.h
--- a/Engine/Source/Runtime/Foliage/Public/InstancedFoliage.h
+++ b/Engine/Source/Runtime/Foliage/Public/InstancedFoliage.h
@@ -202,7 +202,6 @@ struct FFoliageImpl
 	virtual void OnHiddenEditorViewMaskChanged(uint64 InHiddenEditorViews) = 0;
 	virtual void PreEditUndo(AInstancedFoliageActor* IFA, UFoliageType* FoliageType) {}
 	virtual void PostEditUndo(AInstancedFoliageActor* IFA, UFoliageType* FoliageType, const TArray<FFoliageInstance>& Instances, const TSet<int32>& SelectedIndices) = 0;
-	virtual void PostApplyLevelTransform(const FTransform& InTransform, const TArray<FFoliageInstance>& Instances) {}
 	virtual void NotifyFoliageTypeWillChange(AInstancedFoliageActor* IFA, UFoliageType* FoliageType) {}
 	virtual void NotifyFoliageTypeChanged(AInstancedFoliageActor* IFA, UFoliageType* FoliageType, const TArray<FFoliageInstance>& Instances, const TSet<int32>& SelectedIndices, bool bSourceChanged) = 0;
 	virtual void EnterEditMode() {}
