- Fixed Foliage instances to properly handle level editor transform
- ActorFoliage respawning at load (caused by a FoliageType change), actors were spawned in world coordinates, then level editor transform was applied on all level actors. This resulted in bad positioning.
- ActorFoliage respawning at load was also done too early which would lead to respawned actor being invisible
- StaticMeshFoliage in level with editor transform would move some editor instances when saving the level only for instances with a base actors that was moved by the level transform. Visually it didn't cause problems because IFA was also moved, but it would generate a mismatch between editor instances and HISM instances. At this point, erasing foliage and other manipulations using editor data would not work for all cases (loading the level with or without the transform).
- Editor data instances are now always saved local to the Level.
- Also, for now we don't respawn foliage actors in PIE when detecting a mismatch of FoliageType's Actor class to prevent a crash (proper fix will come later)

diff --git a/Engine/Source/Runtime/Engine/Private/LevelUtils.cpp b/Engine/Source/Runtime/Engine/Private/LevelUtils.cpp
--- a/Engine/Source/Runtime/Engine/Private/LevelUtils.cpp
+++ b/Engine/Source/Runtime/Engine/Private/LevelUtils.cpp
@@ -46,6 +46,7 @@ static TMap<ULevel*, FLevelReadOnlyData> LevelReadOnlyCache;
 
 #if WITH_EDITOR
 bool FLevelUtils::bMovingLevel = false;
+bool FLevelUtils::bApplyingLevelTransform = false;
 #endif
 
 /**
@@ -371,14 +372,22 @@ bool FLevelUtils::IsMovingLevel()
 	return bMovingLevel;
 }
 
-#endif // WITH_EDITOR
+bool FLevelUtils::IsApplyingLevelTransform()
+{
+	return bApplyingLevelTransform;
+}
 
+#endif // WITH_EDITOR
 
 void FLevelUtils::ApplyLevelTransform( ULevel* Level, const FTransform& LevelTransform, bool bDoPostEditMove )
 {
 	bool bTransformActors =  !LevelTransform.Equals(FTransform::Identity);
 	if (bTransformActors)
 	{
+#if WITH_EDITOR
+		check(!bApplyingLevelTransform);
+		bApplyingLevelTransform = true;
+#endif
 		if (!LevelTransform.GetRotation().IsIdentity())
 		{
 			// If there is a rotation applied, then the relative precomputed bounds become invalid.
@@ -410,6 +419,9 @@ void FLevelUtils::ApplyLevelTransform( ULevel* Level, const FTransform& LevelTra
 #endif // WITH_EDITOR
 
 		Level->OnApplyLevelTransform.Broadcast(LevelTransform);
+#if WITH_EDITOR
+		bApplyingLevelTransform = false;
+#endif
 	}
 }
 
diff --git a/Engine/Source/Runtime/Engine/Public/LevelUtils.h b/Engine/Source/Runtime/Engine/Public/LevelUtils.h
--- a/Engine/Source/Runtime/Engine/Public/LevelUtils.h
+++ b/Engine/Source/Runtime/Engine/Public/LevelUtils.h
@@ -140,11 +140,13 @@ public:
 	* Returns true if we are moving a level
 	*/
 	static bool IsMovingLevel();
+	static bool IsApplyingLevelTransform();
 
 private:
 
 	// Flag to mark if we are currently finalizing a level offset
 	static bool bMovingLevel;
+	static bool bApplyingLevelTransform;
 
 #endif // WITH_EDITOR
 };
diff --git a/Engine/Source/Runtime/Foliage/Private/FoliageActor.cpp b/Engine/Source/Runtime/Foliage/Private/FoliageActor.cpp
--- a/Engine/Source/Runtime/Foliage/Private/FoliageActor.cpp
+++ b/Engine/Source/Runtime/Foliage/Private/FoliageActor.cpp
@@ -207,11 +207,6 @@ void FFoliageActor::UpdateActorTransforms(const TArray<FFoliageInstance>& Instan
 	}
 }
 
-void FFoliageActor::PostApplyLevelTransform(const FTransform& InTransform, const TArray<FFoliageInstance>& Instances)
-{
-	UpdateActorTransforms(Instances);
-}
-
 void FFoliageActor::PostEditUndo(AInstancedFoliageActor* IFA, UFoliageType* FoliageType, const TArray<FFoliageInstance>& Instances, const TSet<int32>& SelectedIndices)
 {
 	UpdateActorTransforms(Instances);
diff --git a/Engine/Source/Runtime/Foliage/Private/FoliageActor.h b/Engine/Source/Runtime/Foliage/Private/FoliageActor.h
--- a/Engine/Source/Runtime/Foliage/Private/FoliageActor.h
+++ b/Engine/Source/Runtime/Foliage/Private/FoliageActor.h
@@ -41,7 +41,6 @@ struct FFoliageActor : public FFoliageImpl
 	virtual void Refresh(AInstancedFoliageActor* IFA, const TArray<FFoliageInstance>& Instances, bool Async, bool Force) override;
 	virtual void OnHiddenEditorViewMaskChanged(uint64 InHiddenEditorViews) override;
 	virtual void PostEditUndo(AInstancedFoliageActor* IFA, UFoliageType* FoliageType, const TArray<FFoliageInstance>& Instances, const TSet<int32>& SelectedIndices) override;
-	virtual void PostApplyLevelTransform(const FTransform& InTransform, const TArray<FFoliageInstance>& Instances) override;
 	virtual void NotifyFoliageTypeChanged(AInstancedFoliageActor* IFA, UFoliageType* FoliageType, const TArray<FFoliageInstance>& Instances, const TSet<int32>& SelectedIndices, bool bSourceChanged) override;
 	void Reapply(AInstancedFoliageActor* IFA, const UFoliageType* FoliageType, const TArray<FFoliageInstance>& Instances, bool bPostLoad = false);
 	AActor* Spawn(AInstancedFoliageActor* IFA, const FFoliageInstance& Instance);
diff --git a/Engine/Source/Runtime/Foliage/Private/InstancedFoliage.cpp b/Engine/Source/Runtime/Foliage/Private/InstancedFoliage.cpp
--- a/Engine/Source/Runtime/Foliage/Private/InstancedFoliage.cpp
+++ b/Engine/Source/Runtime/Foliage/Private/InstancedFoliage.cpp
@@ -45,6 +45,7 @@ InstancedFoliage.cpp: Instanced foliage implementation.
 #include "UObject/FortniteMainBranchObjectVersion.h"
 #include "PreviewScene.h"
 #include "FoliageActor.h"
+#include "LevelUtils.h"
 
 #define LOCTEXT_NAMESPACE "InstancedFoliage"
 
@@ -2315,6 +2316,13 @@ const FFoliageInfo* AInstancedFoliageActor::FindInfo(const UFoliageType* InType)
 #if WITH_EDITOR
 void AInstancedFoliageActor::MoveInstancesForMovedComponent(UActorComponent* InComponent)
 {
+	// We don't want to handle this case when applying level transform
+	// since it's already handled in AInstancedFoliageActor::OnApplyLevelTransform
+	if (FLevelUtils::IsApplyingLevelTransform())
+	{
+		return;
+	}
+
 	const auto BaseId = InstanceBaseCache.GetInstanceBaseId(InComponent);
 	if (BaseId == FFoliageInstanceBaseCache::InvalidBaseId)
 	{
@@ -3354,6 +3362,9 @@ void AInstancedFoliageActor::PostInitProperties()
 
 		FWorldDelegates::PostApplyLevelOffset.Remove(OnPostApplyLevelOffsetDelegateHandle);
 		OnPostApplyLevelOffsetDelegateHandle = FWorldDelegates::PostApplyLevelOffset.AddUObject(this, &AInstancedFoliageActor::OnPostApplyLevelOffset);
+
+		FWorldDelegates::OnPostWorldInitialization.Remove(OnPostWorldInitializationDelegateHandle);
+		OnPostWorldInitializationDelegateHandle = FWorldDelegates::OnPostWorldInitialization.AddUObject(this, &AInstancedFoliageActor::OnPostWorldInitialization);
 	}
 }
 
@@ -3372,6 +3383,8 @@ void AInstancedFoliageActor::BeginDestroy()
 		}
 
 		FWorldDelegates::PostApplyLevelOffset.Remove(OnPostApplyLevelOffsetDelegateHandle);
+
+		FWorldDelegates::OnPostWorldInitialization.Remove(OnPostWorldInitializationDelegateHandle);
 	}
 }
 #endif
@@ -3530,29 +3543,67 @@ void AInstancedFoliageActor::PostLoad()
 				}
 			}
 
-			// Update foliage component settings if the foliage settings object was changed while the level was not loaded.
-			if (Info.FoliageTypeUpdateGuid != FoliageType->UpdateGuid)
+			// Fixup FoliageInfo instances at load
+			// For Foliage meshes we compute the transforms based on its HISM instances transforms combined with IFA's transform
+			// For Foliage actors we compute the transforms using spawned actors  transform
+			if (GetLinkerCustomVersion(FFoliageCustomVersion::GUID) < FFoliageCustomVersion::FoliageRepairInstancesWithLevelTransform)
 			{
-				if (Info.FoliageTypeUpdateGuid.IsValid())
+				if (Info.Type == EFoliageImplType::StaticMesh)
 				{
-					if (Info.Type == EFoliageImplType::StaticMesh)
+					FFoliageStaticMesh* FoliageStaticMesh = StaticCast<FFoliageStaticMesh*>(Info.Implementation.Get());
+					if (FoliageStaticMesh->Component && (Info.Instances.Num() == FoliageStaticMesh->Component->PerInstanceSMData.Num()))
 					{
-						FFoliageStaticMesh* FoliageStaticMesh = StaticCast<FFoliageStaticMesh*>(Info.Implementation.Get());
-						UFoliageType_InstancedStaticMesh* FoliageType_InstancedStaticMesh = Cast<UFoliageType_InstancedStaticMesh>(FoliageType);
-						FoliageStaticMesh->CheckComponentClass(this, FoliageType_InstancedStaticMesh, Info.Instances, Info.SelectedIndices);
-						FoliageStaticMesh->UpdateComponentSettings(FoliageType_InstancedStaticMesh);
+						if (USceneComponent* IFARootComponent = this->GetRootComponent())
+						{
+							IFARootComponent->UpdateComponentToWorld();
+							FTransform IFATransform = this->GetActorTransform();
+							Info.InstanceHash->Empty();
+							for (int32 InstanceIdx = 0; InstanceIdx < Info.Instances.Num(); InstanceIdx++)
+							{
+								FFoliageInstance& Instance = Info.Instances[InstanceIdx];
+								FTransform Transform = FTransform(FoliageStaticMesh->Component->PerInstanceSMData[InstanceIdx].Transform)*IFATransform;
+								Instance.Location = Transform.GetTranslation();
+								Instance.Rotation = Transform.GetRotation().Rotator();
+								// Rehash instance location
+								Info.InstanceHash->InsertInstance(Instance.Location, InstanceIdx);
+							}
+						}
 					}
-					else if (Info.Type == EFoliageImplType::Actor)
+				}
+				else if (Info.Type == EFoliageImplType::Actor)
+				{
+					FFoliageActor* FoliageActor = StaticCast<FFoliageActor*>(Info.Implementation.Get());
+					if (Info.Instances.Num() == FoliageActor->ActorInstances.Num())
 					{
-						FFoliageActor* FoliageActor = StaticCast<FFoliageActor*>(Info.Implementation.Get());
-						const bool bPostLoad = true;
-						FoliageActor->Reapply(this, FoliageType, Info.Instances, bPostLoad);
+						Info.InstanceHash->Empty();
+						for (int32 InstanceIdx = 0; InstanceIdx < Info.Instances.Num(); InstanceIdx++)
+						{
+							FFoliageInstance& Instance = Info.Instances[InstanceIdx];
+							if (AActor* Actor = FoliageActor->ActorInstances[InstanceIdx])
+							{
+								Actor->ConditionalPostLoad();
+								if (USceneComponent* ActorRootComponent = Actor->GetRootComponent())
+								{
+									ActorRootComponent->UpdateComponentToWorld();
+									FTransform Transform = Actor->GetActorTransform();
+									Instance.Location = Transform.GetTranslation();
+									Instance.Rotation = Transform.GetRotation().Rotator();
+								}
+							}
+							// Rehash instance location
+							Info.InstanceHash->InsertInstance(Instance.Location, InstanceIdx);
+						}
 					}
 				}
-				Info.FoliageTypeUpdateGuid = FoliageType->UpdateGuid;
 			}
 		}
 
+		UWorld* World = GetWorld();
+		if (World && World->bIsWorldInitialized)
+		{
+			DetectFoliageTypeChangeAndUpdate();
+		}
+
 #if WITH_EDITORONLY_DATA
 		if (GetLinkerCustomVersion(FFortniteMainBranchObjectVersion::GUID) < FFortniteMainBranchObjectVersion::FoliageLazyObjPtrToSoftObjPtr)
 		{
@@ -3744,16 +3795,73 @@ void AInstancedFoliageActor::OnLevelActorDeleted(AActor* InActor)
 	}
 }
 
+void AInstancedFoliageActor::OnPostWorldInitialization(UWorld* World, const UWorld::InitializationValues IVS)
+{
+	if (GetWorld() == World)
+	{
+		DetectFoliageTypeChangeAndUpdate();
+	}
+}
+
+// This logic was extracted from AInstancedFoliageActor::PostLoad to be called once the World is done initializing
+void AInstancedFoliageActor::DetectFoliageTypeChangeAndUpdate()
+{
+	for (auto& Pair : FoliageInfos)
+	{
+		// Find the per-mesh info matching the mesh.
+		FFoliageInfo& Info = *Pair.Value;
+		UFoliageType* FoliageType = Pair.Key;
+
+		if (Info.FoliageTypeUpdateGuid != FoliageType->UpdateGuid)
+		{
+			if (Info.FoliageTypeUpdateGuid.IsValid())
+			{
+				// Update foliage component settings if the foliage settings object was changed while the level was not loaded.
+				if (Info.Type == EFoliageImplType::StaticMesh)
+				{
+					FFoliageStaticMesh* FoliageStaticMesh = StaticCast<FFoliageStaticMesh*>(Info.Implementation.Get());
+					UFoliageType_InstancedStaticMesh* FoliageType_InstancedStaticMesh = Cast<UFoliageType_InstancedStaticMesh>(FoliageType);
+					FoliageStaticMesh->CheckComponentClass(this, FoliageType_InstancedStaticMesh, Info.Instances, Info.SelectedIndices);
+					FoliageStaticMesh->UpdateComponentSettings(FoliageType_InstancedStaticMesh);
+				}
+				// Respawn foliage actors
+				else if (Info.Type == EFoliageImplType::Actor)
+				{
+					// We can't spawn in postload because BeingPlay might call UnrealScript which is not supported.
+					UWorld* World = GetWorld();
+					if (World && !World->IsGameWorld())
+					{
+						FFoliageActor* FoliageActor = StaticCast<FFoliageActor*>(Info.Implementation.Get());
+						const bool bPostLoad = true;
+						FoliageActor->Reapply(this, FoliageType, Info.Instances, bPostLoad);
+					}
+				}
+			}
+			Info.FoliageTypeUpdateGuid = FoliageType->UpdateGuid;
+		}
+	}
+}
+
 void AInstancedFoliageActor::OnApplyLevelTransform(const FTransform& InTransform)
 {
+#if WITH_EDITORONLY_DATA
+	InstanceBaseCache.UpdateInstanceBaseCachedTransforms();
 	for (auto& Pair : FoliageInfos)
 	{
 		FFoliageInfo& Info = *Pair.Value;
-		if (Info.Implementation)
+		Info.InstanceHash->Empty();
+		for (int32 InstanceIdx = 0; InstanceIdx < Info.Instances.Num(); InstanceIdx++)
 		{
-			Info.Implementation->PostApplyLevelTransform(InTransform, Info.Instances);
+			FFoliageInstance& Instance = Info.Instances[InstanceIdx];
+			FTransform OldTransform(Instance.Rotation, Instance.Location);
+			FTransform NewTransform = OldTransform * InTransform;
+			Instance.Location = NewTransform.GetTranslation();
+			Instance.Rotation = NewTransform.Rotator();
+			// Rehash instance location
+			Info.InstanceHash->InsertInstance(Instance.Location, InstanceIdx);
 		}
 	}
+#endif
 }
 
 void AInstancedFoliageActor::OnPostApplyLevelOffset(ULevel* InLevel, UWorld* InWorld, const FVector& InOffset, bool bWorldShift)
diff --git a/Engine/Source/Runtime/Foliage/Public/InstancedFoliage.h b/Engine/Source/Runtime/Foliage/Public/InstancedFoliage.h
--- a/Engine/Source/Runtime/Foliage/Public/InstancedFoliage.h
+++ b/Engine/Source/Runtime/Foliage/Public/InstancedFoliage.h
@@ -202,7 +202,6 @@ struct FFoliageImpl
 	virtual void OnHiddenEditorViewMaskChanged(uint64 InHiddenEditorViews) = 0;
 	virtual void PreEditUndo(AInstancedFoliageActor* IFA, UFoliageType* FoliageType) {}
 	virtual void PostEditUndo(AInstancedFoliageActor* IFA, UFoliageType* FoliageType, const TArray<FFoliageInstance>& Instances, const TSet<int32>& SelectedIndices) = 0;
-	virtual void PostApplyLevelTransform(const FTransform& InTransform, const TArray<FFoliageInstance>& Instances) {}
 	virtual void NotifyFoliageTypeWillChange(AInstancedFoliageActor* IFA, UFoliageType* FoliageType) {}
 	virtual void NotifyFoliageTypeChanged(AInstancedFoliageActor* IFA, UFoliageType* FoliageType, const TArray<FFoliageInstance>& Instances, const TSet<int32>& SelectedIndices, bool bSourceChanged) = 0;
 	virtual void EnterEditMode() {}
diff --git a/Engine/Source/Runtime/Foliage/Public/InstancedFoliageActor.h b/Engine/Source/Runtime/Foliage/Public/InstancedFoliageActor.h
--- a/Engine/Source/Runtime/Foliage/Public/InstancedFoliageActor.h
+++ b/Engine/Source/Runtime/Foliage/Public/InstancedFoliageActor.h
@@ -46,6 +46,8 @@ struct FFoliageCustomVersion
 		FoliageActorSupport = 12,
 		// Foliage Actor (No weak ptr)
 		FoliageActorSupportNoWeakPtr = 13,
+		// Foliage Instances are now always saved local to Level
+		FoliageRepairInstancesWithLevelTransform = 14,
 		// -----<new versions can be added above this line>-------------------------------------------------
 		VersionPlusOne,
 		LatestVersion = VersionPlusOne - 1
@@ -98,6 +100,7 @@ public:
 	virtual void BeginDestroy() override;
 	virtual void Destroyed() override;
 	FOLIAGE_API void CleanupDeletedFoliageType();
+	FOLIAGE_API void DetectFoliageTypeChangeAndUpdate();
 
 	// Delegate type for selection change events
 	DECLARE_MULTICAST_DELEGATE_TwoParams(FOnSelectionChanged, bool, const TArray<AActor*>&);
@@ -253,6 +256,7 @@ private:
 	void OnLevelActorDeleted(AActor* InActor);
 	void OnApplyLevelTransform(const FTransform& InTransform);
 	void OnPostApplyLevelOffset(ULevel* InLevel, UWorld* InWorld, const FVector& InOffset, bool bWorldShift);
+	void OnPostWorldInitialization(UWorld* World, const UWorld::InitializationValues IVS);
 
 	// Move instances to a foliage actor in target level
 	FOLIAGE_API void MoveInstancesToLevel(ULevel* InTargetLevel, TSet<int32>& InInstanceList, FFoliageInfo* InCurrentMeshInfo, UFoliageType* InFoliageType);
@@ -263,6 +267,7 @@ private:
 	FDelegateHandle OnLevelActorDeletedDelegateHandle;
 	FDelegateHandle OnPostApplyLevelOffsetDelegateHandle;
 	FDelegateHandle OnApplyLevelTransformDelegateHandle;
+	FDelegateHandle OnPostWorldInitializationDelegateHandle;
 
 	FOnFoliageTypeMeshChanged OnFoliageTypeMeshChangedEvent;
 #endif
