Fixed Development EditorGame target
#rnx

diff --git a/Engine/Source/Runtime/Foliage/Private/InstancedFoliage.cpp b/Engine/Source/Runtime/Foliage/Private/InstancedFoliage.cpp
--- a/Engine/Source/Runtime/Foliage/Private/InstancedFoliage.cpp
+++ b/Engine/Source/Runtime/Foliage/Private/InstancedFoliage.cpp
@@ -3845,20 +3845,23 @@ void AInstancedFoliageActor::DetectFoliageTypeChangeAndUpdate()
 void AInstancedFoliageActor::OnApplyLevelTransform(const FTransform& InTransform)
 {
 #if WITH_EDITORONLY_DATA
-	InstanceBaseCache.UpdateInstanceBaseCachedTransforms();
-	for (auto& Pair : FoliageInfos)
+	if (GIsEditor)
 	{
-		FFoliageInfo& Info = *Pair.Value;
-		Info.InstanceHash->Empty();
-		for (int32 InstanceIdx = 0; InstanceIdx < Info.Instances.Num(); InstanceIdx++)
+		InstanceBaseCache.UpdateInstanceBaseCachedTransforms();
+		for (auto& Pair : FoliageInfos)
 		{
-			FFoliageInstance& Instance = Info.Instances[InstanceIdx];
-			FTransform OldTransform(Instance.Rotation, Instance.Location);
-			FTransform NewTransform = OldTransform * InTransform;
-			Instance.Location = NewTransform.GetTranslation();
-			Instance.Rotation = NewTransform.Rotator();
-			// Rehash instance location
-			Info.InstanceHash->InsertInstance(Instance.Location, InstanceIdx);
+			FFoliageInfo& Info = *Pair.Value;
+			Info.InstanceHash->Empty();
+			for (int32 InstanceIdx = 0; InstanceIdx < Info.Instances.Num(); InstanceIdx++)
+			{
+				FFoliageInstance& Instance = Info.Instances[InstanceIdx];
+				FTransform OldTransform(Instance.Rotation, Instance.Location);
+				FTransform NewTransform = OldTransform * InTransform;
+				Instance.Location = NewTransform.GetTranslation();
+				Instance.Rotation = NewTransform.Rotator();
+				// Rehash instance location
+				Info.InstanceHash->InsertInstance(Instance.Location, InstanceIdx);
+			}
 		}
 	}
 #endif
