Add 0 initlizers to RendererSettings enums and ESimulationQuery

The following enums are effected:

- EMobileMSAASampleCount
- ECompositingSampleCount
- ESimulationQuery

#rnx

diff --git a/Engine/Source/Runtime/Foliage/Public/ProceduralFoliageInstance.h b/Engine/Source/Runtime/Foliage/Public/ProceduralFoliageInstance.h
--- a/Engine/Source/Runtime/Foliage/Public/ProceduralFoliageInstance.h
+++ b/Engine/Source/Runtime/Foliage/Public/ProceduralFoliageInstance.h
@@ -37,6 +37,7 @@ namespace ESimulationQuery
 {
 	enum Type
 	{
+		None = 0 UMETA(Hidden),
 		/*Instances overlap with collision*/
 		CollisionOverlap = 1,
 		/*Instances overlap with shade*/
