PR #5412: UE-61103: Load BodyInstance data in FoliageType (Contributed by [REDACTED_PERSON])

diff --git a/Engine/Source/Runtime/Foliage/Private/InstancedFoliage.cpp b/Engine/Source/Runtime/Foliage/Private/InstancedFoliage.cpp
--- a/Engine/Source/Runtime/Foliage/Private/InstancedFoliage.cpp
+++ b/Engine/Source/Runtime/Foliage/Private/InstancedFoliage.cpp
@@ -489,6 +489,16 @@ void UFoliageType::Serialize(FArchive& Ar)
 #endif// WITH_EDITORONLY_DATA
 }
 
+void UFoliageType::PostLoad()
+{
+	Super::PostLoad();
+
+	if (!IsTemplate())
+	{
+		BodyInstance.FixupData(this);
+	}
+}
+
 bool UFoliageType::IsNotAssetOrBlueprint() const
 {
 	return IsAsset() == false && Cast<UBlueprint>(GetClass()->ClassGeneratedBy) == nullptr;
diff --git a/Engine/Source/Runtime/Foliage/Public/FoliageType.h b/Engine/Source/Runtime/Foliage/Public/FoliageType.h
--- a/Engine/Source/Runtime/Foliage/Public/FoliageType.h
+++ b/Engine/Source/Runtime/Foliage/Public/FoliageType.h
@@ -92,6 +92,7 @@ class UFoliageType : public UObject
 	virtual UClass* GetComponentClass() const PURE_VIRTUAL(UFoliageType::GetComponentClass(), return nullptr;);
 	
 	virtual void Serialize(FArchive& Ar) override;
+	virtual void PostLoad() override;
 
 	virtual bool IsNotAssetOrBlueprint() const;
 
