Resubmitted cl 8576092 without circular dependency
(Fixed 'Foliage' Engine Show Flag ignored by Actors painted by [REDACTED_PERSON])

diff --git a/Engine/Source/Editor/FoliageEdit/Private/FoliageEdMode.cpp b/Engine/Source/Editor/FoliageEdit/Private/FoliageEdMode.cpp
--- a/Engine/Source/Editor/FoliageEdit/Private/FoliageEdMode.cpp
+++ b/Engine/Source/Editor/FoliageEdit/Private/FoliageEdMode.cpp
@@ -29,6 +29,7 @@
 #include "LevelEditor.h"
 #include "Toolkits/ToolkitManager.h"
 #include "FoliageEditActions.h"
+#include "FoliageHelper.h"
 
 #include "AssetRegistryModule.h"
 #include "Misc/ScopeExit.h"
@@ -176,7 +177,7 @@ bool FFoliagePaintingGeometryFilter::operator() (const UPrimitiveComponent* Comp
 {
 	if (Component)
 	{
-		bool bFoliageOwned = Component->GetOwner() && AInstancedFoliageActor::IsOwnedByFoliage(Component->GetOwner());
+		bool bFoliageOwned = Component->GetOwner() && FFoliageHelper::IsOwnedByFoliage(Component->GetOwner());
 
 		// Whitelist
 		bool bAllowed =
@@ -3584,7 +3585,7 @@ void FEdModeFoliage::DrawHUD(FEditorViewportClient* ViewportClient, FViewport* V
 /** FEdMode: Check to see if an actor can be selected in this mode - no side effects */
 bool FEdModeFoliage::IsSelectionAllowed(AActor* InActor, bool bInSelection) const
 {
-	return AInstancedFoliageActor::IsOwnedByFoliage(InActor);
+	return FFoliageHelper::IsOwnedByFoliage(InActor);
 }
 
 /** FEdMode: Handling SelectActor */
@@ -3643,7 +3644,7 @@ bool FEdModeFoliage::HandleClick(FEditorViewportClient* InViewportClient, HHitPr
 				UpdateWidgetLocationToInstanceSelection();
 			}
 		}
-		else if (HitProxy && HitProxy->IsA(HActor::StaticGetType()) && AInstancedFoliageActor::IsOwnedByFoliage(((HActor*)HitProxy)->Actor))
+		else if (HitProxy && HitProxy->IsA(HActor::StaticGetType()) && FFoliageHelper::IsOwnedByFoliage(((HActor*)HitProxy)->Actor))
 		{
 			HActor* ActorProxy = ((HActor*)HitProxy);
 			AInstancedFoliageActor* IFA = AInstancedFoliageActor::GetInstancedFoliageActorForLevel(ActorProxy->Actor->GetLevel());
diff --git a/Engine/Source/Runtime/Engine/Engine.Build.cs b/Engine/Source/Runtime/Engine/Engine.Build.cs
--- a/Engine/Source/Runtime/Engine/Engine.Build.cs
+++ b/Engine/Source/Runtime/Engine/Engine.Build.cs
@@ -36,7 +36,7 @@ public class Engine : ModuleRules
 				"AutomationWorker",
 				"MovieSceneCapture",
 				"DesktopPlatform",
-			}
+            }
 		);
 
 		if (Target.Configuration != UnrealTargetConfiguration.Shipping)
@@ -187,11 +187,16 @@ public class Engine : ModuleRules
 
 		CircularlyReferencedDependentModules.Add("GameplayTags");
 		CircularlyReferencedDependentModules.Add("Landscape");
-		CircularlyReferencedDependentModules.Add("UMG");
+        CircularlyReferencedDependentModules.Add("UMG");
 		CircularlyReferencedDependentModules.Add("MaterialShaderQualitySettings");
 		CircularlyReferencedDependentModules.Add("CinematicCamera");
 		CircularlyReferencedDependentModules.Add("AudioMixer");
 
+		if (Target.Type == TargetType.Editor)
+		{
+			PrivateIncludePathModuleNames.Add("Foliage");
+		}
+
 		// The AnimGraphRuntime module is not needed by Engine proper, but it is loaded in LaunchEngineLoop.cpp,
 		// and needs to be listed in an always-included module in order to be compiled into standalone games
 		DynamicallyLoadedModuleNames.Add("AnimGraphRuntime");
diff --git a/Engine/Source/Runtime/Engine/Private/PrimitiveSceneProxy.cpp b/Engine/Source/Runtime/Engine/Private/PrimitiveSceneProxy.cpp
--- a/Engine/Source/Runtime/Engine/Private/PrimitiveSceneProxy.cpp
+++ b/Engine/Source/Runtime/Engine/Private/PrimitiveSceneProxy.cpp
@@ -15,6 +15,9 @@
 #include "Materials/Material.h"
 #include "SceneManagement.h"
 #include "VT/RuntimeVirtualTexture.h"
+#if WITH_EDITOR
+#include "FoliageHelper.h"
+#endif
 
 static TAutoConsoleVariable<int32> CVarForceSingleSampleShadowingFromStationary(
 	TEXT("r.Shadow.ForceSingleSampleShadowingFromStationary"),
@@ -160,6 +163,7 @@ FPrimitiveSceneProxy::FPrimitiveSceneProxy(const UPrimitiveComponent* InComponen
 // by default we are always drawn
 ,	HiddenEditorViews(0)
 ,	DrawInAnyEditMode(0)
+,   bIsFoliage(false)
 #endif
 ,	VisibilityId(InComponent->VisibilityId)
 ,	MaxDrawDistance(InComponent->CachedMaxDrawDistance > 0 ? InComponent->CachedMaxDrawDistance : FLT_MAX)
@@ -209,6 +213,7 @@ FPrimitiveSceneProxy::FPrimitiveSceneProxy(const UPrimitiveComponent* InComponen
 		// cache the actor's group membership
 		HiddenEditorViews = InComponent->GetHiddenEditorViews();
 		DrawInAnyEditMode = InComponent->GetOwner()->IsEditorOnly();
+		bIsFoliage = FFoliageHelper::IsOwnedByFoliage(InComponent->GetOwner());
 #endif
 	}
 	
@@ -619,6 +624,11 @@ bool FPrimitiveSceneProxy::IsShown(const FSceneView* View) const
 		}
 	}
 
+	if (bIsFoliage && !View->Family->EngineShowFlags.InstancedFoliage)
+	{
+		return false;
+	}
+
 	// After checking for VR/Desktop Edit mode specific actors, check for Editor vs. Game
 	if(View->Family->EngineShowFlags.Editor)
 	{
diff --git a/Engine/Source/Runtime/Engine/Public/PrimitiveSceneProxy.h b/Engine/Source/Runtime/Engine/Public/PrimitiveSceneProxy.h
--- a/Engine/Source/Runtime/Engine/Public/PrimitiveSceneProxy.h
+++ b/Engine/Source/Runtime/Engine/Public/PrimitiveSceneProxy.h
@@ -1042,6 +1042,8 @@ private:
 
 	/** Whether this should only draw in any editing mode*/
 	uint32 DrawInAnyEditMode : 1;
+
+	uint32 bIsFoliage : 1;
 #endif
 
 	/** Used for precomputed visibility */
diff --git a/Engine/Source/Runtime/Foliage/Private/FoliageActor.cpp b/Engine/Source/Runtime/Foliage/Private/FoliageActor.cpp
--- a/Engine/Source/Runtime/Foliage/Private/FoliageActor.cpp
+++ b/Engine/Source/Runtime/Foliage/Private/FoliageActor.cpp
@@ -3,8 +3,7 @@
 #include "FoliageActor.h"
 #include "InstancedFoliageActor.h"
 #include "FoliageType_Actor.h"
-
-FName FoliageActorTag(TEXT("FoliageActorInstance"));
+#include "FoliageHelper.h"
 
 //
 //
@@ -87,7 +86,7 @@ AActor* FFoliageActor::Spawn(AInstancedFoliageActor* IFA, const FFoliageInstance
 	AActor* NewActor = IFA->GetWorld()->SpawnActor<AActor>(ActorClass, Instance.GetInstanceWorldTransform(), SpawnParameters);
 	if (NewActor)
 	{
-		NewActor->Tags.AddUnique(FoliageActorTag);
+		FFoliageHelper::SetIsOwnedByFoliage(NewActor);
 	}
 	return NewActor;
 }
diff --git a/Engine/Source/Runtime/Foliage/Private/InstancedFoliage.cpp b/Engine/Source/Runtime/Foliage/Private/InstancedFoliage.cpp
--- a/Engine/Source/Runtime/Foliage/Private/InstancedFoliage.cpp
+++ b/Engine/Source/Runtime/Foliage/Private/InstancedFoliage.cpp
@@ -46,6 +46,7 @@ InstancedFoliage.cpp: Instanced foliage implementation.
 #include "PreviewScene.h"
 #include "FoliageActor.h"
 #include "LevelUtils.h"
+#include "FoliageHelper.h"
 
 #define LOCTEXT_NAMESPACE "InstancedFoliage"
 
@@ -59,8 +60,6 @@ DECLARE_CYCLE_STAT(TEXT("FoliageMeshInfo_AddInstance"), STAT_FoliageAddInstance,
 DECLARE_CYCLE_STAT(TEXT("FoliageMeshInfo_RemoveInstance"), STAT_FoliageRemoveInstance, STATGROUP_Foliage);
 DECLARE_CYCLE_STAT(TEXT("FoliageMeshInfo_CreateComponent"), STAT_FoliageCreateComponent, STATGROUP_Foliage);
 
-extern FName FoliageActorTag;
-
 static TAutoConsoleVariable<int32> CVarFoliageDiscardDataOnLoad(
 	TEXT("foliage.DiscardDataOnLoad"),
 	0,
@@ -2168,11 +2167,6 @@ AInstancedFoliageActor::AInstancedFoliageActor(const FObjectInitializer& ObjectI
 	PrimaryActorTick.bCanEverTick = false;
 }
 
-bool AInstancedFoliageActor::IsOwnedByFoliage(const AActor* InActor)
-{
-	return InActor != nullptr && InActor->ActorHasTag(FoliageActorTag);
-}
-
 AInstancedFoliageActor* AInstancedFoliageActor::GetInstancedFoliageActorForCurrentLevel(UWorld* InWorld, bool bCreateIfNone)
 {
 	return GetInstancedFoliageActorForLevel(InWorld->GetCurrentLevel(), bCreateIfNone);
@@ -3999,7 +3993,7 @@ bool AInstancedFoliageActor::FoliageTrace(const UWorld* InWorld, FHitResult& Out
 
 		// Don't place foliage on itself
 		const AInstancedFoliageActor* FoliageActor = Cast<AInstancedFoliageActor>(HitActor);
-		if (!FoliageActor && HitActor && AInstancedFoliageActor::IsOwnedByFoliage(HitActor))
+		if (!FoliageActor && HitActor && FFoliageHelper::IsOwnedByFoliage(HitActor))
 		{
 			FoliageActor = HitActor->GetLevel()->InstancedFoliageActor.Get();
 			if (FoliageActor == nullptr)
diff --git a/Engine/Source/Runtime/Foliage/Public/FoliageHelper.h b/Engine/Source/Runtime/Foliage/Public/FoliageHelper.h
new file mode 100644
--- /dev/null
+++ b/Engine/Source/Runtime/Foliage/Public/FoliageHelper.h
@@ -0,0 +1,17 @@
+// Copyright 1998-2019 Epic Games, Inc. All Rights Reserved.
+
+#pragma once
+
+#include "CoreMinimal.h"
+
+/**
+ * Foliage Helper class
+ */
+class FFoliageHelper
+{
+public:
+#if WITH_EDITOR
+	static void SetIsOwnedByFoliage(AActor* InActor) { if (InActor) { InActor->Tags.AddUnique("FoliageActorInstance"); } }
+	static bool IsOwnedByFoliage(const AActor* InActor) { return InActor != nullptr && InActor->ActorHasTag("FoliageActorInstance"); }
+#endif
+};
\ No newline at end of file
diff --git a/Engine/Source/Runtime/Foliage/Public/InstancedFoliageActor.h b/Engine/Source/Runtime/Foliage/Public/InstancedFoliageActor.h
--- a/Engine/Source/Runtime/Foliage/Public/InstancedFoliageActor.h
+++ b/Engine/Source/Runtime/Foliage/Public/InstancedFoliageActor.h
@@ -90,7 +90,6 @@ public:
 	virtual void RerunConstructionScripts() override {}
 	virtual bool IsLevelBoundsRelevant() const override { return false; }
 
-	FOLIAGE_API static bool IsOwnedByFoliage(const AActor* InActor);
 protected:
 	// Default InternalTakeRadialDamage behavior finds and scales damage for the closest component which isn't appropriate for foliage.
 	virtual float InternalTakeRadialDamage(float Damage, struct FRadialDamageEvent const& RadialDamageEvent, class AController* EventInstigator, AActor* DamageCauser) override;
