Support Foliage Actor with No Base Component
Fix Snap back issue when moving Actor outsdie of Foliage Mode
Delete Foliage Instances when Actor is deleted outside of Foliage Mode

#rnx

diff --git a/Engine/Source/Editor/FoliageEdit/Private/FoliageEdMode.cpp b/Engine/Source/Editor/FoliageEdit/Private/FoliageEdMode.cpp
--- a/Engine/Source/Editor/FoliageEdit/Private/FoliageEdMode.cpp
+++ b/Engine/Source/Editor/FoliageEdit/Private/FoliageEdMode.cpp
@@ -1971,14 +1971,14 @@ void FEdModeFoliage::SelectInvalidInstances(const UFoliageType* Settings)
 		TArray<FHitResult> Hits; Hits.Reserve(16);
 
 		TArray<int32> InvalidInstances;
-
+		
 		for (int32 InstanceIdx = 0; InstanceIdx < NumInstances; ++InstanceIdx)
 		{
 			FFoliageInstance& Instance = FoliageInfo->Instances[InstanceIdx];
 			UActorComponent* CurrentInstanceBase = IFA->InstanceBaseCache.GetInstanceBasePtr(Instance.BaseId).Get();
-			bool bInvalidInstance = true;
-
-			if (CurrentInstanceBase != nullptr)
+			
+			bool bInvalidInstance = FoliageInfo->ShouldAttachToBaseComponent() || (!FoliageInfo->ShouldAttachToBaseComponent() && CurrentInstanceBase != nullptr);
+			if (FoliageInfo->ShouldAttachToBaseComponent() && CurrentInstanceBase != nullptr)
 			{
 				FVector InstanceTraceRange = Instance.GetInstanceWorldTransform().TransformVector(FVector(0.f, 0.f, 1000.f));
 				FVector Start = Instance.Location + InstanceTraceRange;
@@ -2015,6 +2015,7 @@ void FEdModeFoliage::SelectInvalidInstances(const UFoliageType* Settings)
 						if ((DistanceToGround - InstanceWorldTreshold) <= KINDA_SMALL_NUMBER)
 						{
 							bInvalidInstance = false;
+							break;
 						}
 					}
 				}
@@ -2987,9 +2988,13 @@ bool FEdModeFoliage::SnapInstanceToGround(AInstancedFoliageActor* InIFA, float A
 		}
 
 		// Set new base
-		auto NewBaseId = InIFA->InstanceBaseCache.AddInstanceBaseId(HitComponent);
+		auto NewBaseId = InIFA->InstanceBaseCache.AddInstanceBaseId(Mesh.ShouldAttachToBaseComponent() ? HitComponent : nullptr);
 		Mesh.RemoveFromBaseHash(InstanceIdx);
 		Instance.BaseId = NewBaseId;
+		if (Instance.BaseId == FFoliageInstanceBaseCache::InvalidBaseId)
+		{
+			Instance.BaseComponent = nullptr;
+		}
 		Mesh.AddToBaseHash(InstanceIdx);
 		Instance.Location = Hit.Location;
 		Instance.ZOffset = 0.f;
diff --git a/Engine/Source/Runtime/Foliage/Private/FoliageActor.cpp b/Engine/Source/Runtime/Foliage/Private/FoliageActor.cpp
--- a/Engine/Source/Runtime/Foliage/Private/FoliageActor.cpp
+++ b/Engine/Source/Runtime/Foliage/Private/FoliageActor.cpp
@@ -63,6 +63,7 @@ void FFoliageActor::Initialize(AInstancedFoliageActor* IFA, const UFoliageType*
 	check(!IsInitialized());
 	const UFoliageType_Actor* FoliageType_Actor = Cast<UFoliageType_Actor>(FoliageType);
 	ActorClass = FoliageType_Actor->ActorClass ? FoliageType_Actor->ActorClass.Get() : AActor::StaticClass();
+	bShouldAttachToBaseComponent = FoliageType_Actor->bShouldAttachToBaseComponent;
 }
 
 void FFoliageActor::Uninitialize()
@@ -213,6 +214,18 @@ void FFoliageActor::PostEditUndo(AInstancedFoliageActor* IFA, UFoliageType* Foli
 
 void FFoliageActor::NotifyFoliageTypeChanged(AInstancedFoliageActor* IFA, UFoliageType* FoliageType, const TArray<FFoliageInstance>& Instances, const TSet<int32>& SelectedIndices, bool bSourceChanged)
 {
+	if (UFoliageType_Actor* FoliageTypeActor = Cast<UFoliageType_Actor>(FoliageType))
+	{
+		if (bShouldAttachToBaseComponent != FoliageTypeActor->bShouldAttachToBaseComponent)
+		{
+			bShouldAttachToBaseComponent = FoliageTypeActor->bShouldAttachToBaseComponent;
+			if (!bShouldAttachToBaseComponent)
+			{
+				IFA->RemoveBaseComponentOnFoliageTypeInstances(FoliageType);
+			}
+		}
+	}
+
 	if (bSourceChanged)
 	{
 		Reapply(IFA, FoliageType, Instances);
@@ -265,4 +278,35 @@ void FFoliageActor::ClearSelection(const TSet<int32>& SelectedIndices)
 	AInstancedFoliageActor::SelectionChanged.Broadcast(false, GetActorsFromSelectedIndices(SelectedIndices));
 }
 
-#endif // WITH_EDITOR
\ No newline at end of file
+bool FFoliageActor::UpdateInstanceFromActor(AInstancedFoliageActor* IFA, AActor* InActor, FFoliageInfo& FoliageInfo)
+{
+	int32 Index = FindIndex(InActor);
+	if (Index == INDEX_NONE)
+	{
+		return false;
+	}
+
+	IFA->Modify();
+	const bool bChecked = false; // In the case of the PostEditUndo its possible that the instancehash is empty.
+	FoliageInfo.InstanceHash->RemoveInstance(FoliageInfo.Instances[Index].Location, Index, bChecked);
+	
+	FTransform ActorTransform = InActor->GetTransform();
+	FoliageInfo.Instances[Index].Location = ActorTransform.GetLocation();
+	FoliageInfo.Instances[Index].Rotation = FRotator(ActorTransform.GetRotation());
+	FoliageInfo.InstanceHash->InsertInstance(FoliageInfo.Instances[Index].Location, Index);
+	
+	return true;
+}
+
+void FFoliageActor::GetInvalidInstances(TArray<int32>& InvalidInstances)
+{
+	for (int32 Index = 0; Index < ActorInstances.Num(); ++Index)
+	{
+		if (ActorInstances[Index] == nullptr)
+		{
+			InvalidInstances.Add(Index);
+		}
+	}
+}
+
+#endif // WITH_EDITOR
diff --git a/Engine/Source/Runtime/Foliage/Private/FoliageActor.h b/Engine/Source/Runtime/Foliage/Private/FoliageActor.h
--- a/Engine/Source/Runtime/Foliage/Private/FoliageActor.h
+++ b/Engine/Source/Runtime/Foliage/Private/FoliageActor.h
@@ -11,9 +11,11 @@ struct FFoliageActor : public FFoliageImpl
 
 	TArray<AActor*> ActorInstances;
 	UClass* ActorClass;
+	bool bShouldAttachToBaseComponent;
 	
 	FFoliageActor()
 		: ActorClass(nullptr)
+		, bShouldAttachToBaseComponent(true)
 	{
 	}
 
@@ -45,6 +47,9 @@ struct FFoliageActor : public FFoliageImpl
 	void Reapply(AInstancedFoliageActor* IFA, const UFoliageType* FoliageType, const TArray<FFoliageInstance>& Instances, bool bPostLoad = false);
 	AActor* Spawn(AInstancedFoliageActor* IFA, const FFoliageInstance& Instance);
 	TArray<AActor*> GetActorsFromSelectedIndices(const TSet<int32>& SelectedIndices) const;
+	virtual bool ShouldAttachToBaseComponent() const override { return bShouldAttachToBaseComponent; }
+	bool UpdateInstanceFromActor(AInstancedFoliageActor* IFA, AActor* InActor, FFoliageInfo& FoliageInfo);
+	void GetInvalidInstances(TArray<int32>& InvalidInstances);
 
 private:
 	void UpdateActorTransforms(const TArray<FFoliageInstance>& Instances);
diff --git a/Engine/Source/Runtime/Foliage/Private/InstancedFoliage.cpp b/Engine/Source/Runtime/Foliage/Private/InstancedFoliage.cpp
--- a/Engine/Source/Runtime/Foliage/Private/InstancedFoliage.cpp
+++ b/Engine/Source/Runtime/Foliage/Private/InstancedFoliage.cpp
@@ -696,6 +696,7 @@ UFoliageType_Actor::UFoliageType_Actor(const FObjectInitializer& ObjectInitializ
 {
 	Density = 10;
 	Radius = 500;
+	bShouldAttachToBaseComponent = true;
 }
 
 #if WITH_EDITOR
@@ -1514,6 +1515,7 @@ FFoliageInfo::FFoliageInfo()
 	: Type(EFoliageImplType::StaticMesh)
 #if WITH_EDITOR
 	, InstanceHash(GIsEditor ? new FFoliageInstanceHash() : nullptr)
+	, bMovingInstances(false)
 #endif
 { }
 
@@ -1710,7 +1712,11 @@ void FFoliageInfo::AddInstanceImpl(AInstancedFoliageActor* InIFA, const FFoliage
 	int32 InstanceIndex = Instances.Add(InNewInstance);
 	FFoliageInstance& AddedInstance = Instances[InstanceIndex];
 
-	AddedInstance.BaseId = InIFA->InstanceBaseCache.AddInstanceBaseId(InNewInstance.BaseComponent);
+	AddedInstance.BaseId = InIFA->InstanceBaseCache.AddInstanceBaseId(ShouldAttachToBaseComponent()? InNewInstance.BaseComponent : nullptr);
+	if (AddedInstance.BaseId == FFoliageInstanceBaseCache::InvalidBaseId)
+	{
+		AddedInstance.BaseComponent = nullptr;
+	}
 
 	// Add the instance to the hash
 	AddToBaseHash(InstanceIndex);
@@ -1745,7 +1751,11 @@ void FFoliageInfo::AddInstances(AInstancedFoliageActor* InIFA, const UFoliageTyp
 void FFoliageInfo::AddInstance(AInstancedFoliageActor* InIFA, const UFoliageType* InSettings, const FFoliageInstance& InNewInstance, UActorComponent* InBaseComponent)
 {
 	FFoliageInstance Instance = InNewInstance;
-	Instance.BaseId = InIFA->InstanceBaseCache.AddInstanceBaseId(InBaseComponent);
+	Instance.BaseId = InIFA->InstanceBaseCache.AddInstanceBaseId(ShouldAttachToBaseComponent() ? InBaseComponent : nullptr);
+	if (Instance.BaseId == FFoliageInstanceBaseCache::InvalidBaseId)
+	{
+		Instance.BaseComponent = nullptr;
+	}
 	AddInstance(InIFA, InSettings, Instance);
 }
 
@@ -1855,6 +1865,8 @@ void FFoliageInfo::RemoveInstances(AInstancedFoliageActor* InIFA, const TArray<i
 
 void FFoliageInfo::PreMoveInstances(AInstancedFoliageActor* InIFA, const TArray<int32>& InInstancesToMove)
 {
+	bMovingInstances = true;
+
 	// Remove instances from the hash
 	for (TArray<int32>::TConstIterator It(InInstancesToMove); It; ++It)
 	{
@@ -1898,6 +1910,8 @@ void FFoliageInfo::PostUpdateInstances(AInstancedFoliageActor* InIFA, const TArr
 void FFoliageInfo::PostMoveInstances(AInstancedFoliageActor* InIFA, const TArray<int32>& InInstancesMoved)
 {
 	PostUpdateInstances(InIFA, InInstancesMoved, true, true);
+
+	bMovingInstances = false;
 }
 
 void FFoliageInfo::DuplicateInstances(AInstancedFoliageActor* InIFA, UFoliageType* InSettings, const TArray<int32>& InInstancesToDuplicate)
@@ -2144,6 +2158,17 @@ void FFoliageInfo::ExitEditMode()
 	Implementation->ExitEditMode();
 }
 
+void FFoliageInfo::RemoveBaseComponentOnInstances()
+{
+	for (int32 InstanceIdx = 0; InstanceIdx < Instances.Num(); InstanceIdx++)
+	{
+		RemoveFromBaseHash(InstanceIdx);
+		Instances[InstanceIdx].BaseId = FFoliageInstanceBaseCache::InvalidBaseId;
+		Instances[InstanceIdx].BaseComponent = nullptr;
+		AddToBaseHash(InstanceIdx);
+	}
+}
+
 TArray<int32> FFoliageInfo::GetInstancesOverlappingBox(const FBox& Box) const
 {
 	return InstanceHash->GetInstancesOverlappingBox(Box);
@@ -2308,7 +2333,7 @@ const FFoliageInfo* AInstancedFoliageActor::FindInfo(const UFoliageType* InType)
 
 
 #if WITH_EDITOR
-void AInstancedFoliageActor::MoveInstancesForMovedComponent(UActorComponent* InComponent)
+void AInstancedFoliageActor::MoveInstancesForMovedOwnedActors(AActor* InActor)
 {
 	// We don't want to handle this case when applying level transform
 	// since it's already handled in AInstancedFoliageActor::OnApplyLevelTransform
@@ -2316,7 +2341,37 @@ void AInstancedFoliageActor::MoveInstancesForMovedComponent(UActorComponent* InC
 	{
 		return;
 	}
+		
+	for (auto& Pair : FoliageInfos)
+	{
+		// Source of movement is the Foliage
+		if (Pair.Value->bMovingInstances)
+		{
+			continue;
+		}
+
+		if (Pair.Key->IsA<UFoliageType_Actor>())
+		{
+			if (FFoliageActor* FoliageActor = StaticCast<FFoliageActor*>(Pair.Value->Implementation.Get()))
+			{
+				if (FoliageActor->UpdateInstanceFromActor(this, InActor, *Pair.Value))
+				{
+					return;
+				}
+			}
+		}
+	}
+}
 
+void AInstancedFoliageActor::MoveInstancesForMovedComponent(UActorComponent* InComponent)
+{
+	// We don't want to handle this case when applying level transform
+	// since it's already handled in AInstancedFoliageActor::OnApplyLevelTransform
+	if (FLevelUtils::IsApplyingLevelTransform())
+	{
+		return;
+	}
+		
 	const auto BaseId = InstanceBaseCache.GetInstanceBaseId(InComponent);
 	if (BaseId == FFoliageInstanceBaseCache::InvalidBaseId)
 	{
@@ -3480,6 +3535,12 @@ void AInstancedFoliageActor::PostLoad()
 			Info.InstanceHash->Empty();
 			for (int32 InstanceIdx = 0; InstanceIdx < Info.Instances.Num(); InstanceIdx++)
 			{
+				// Invalidate base if we aren't supposed to be attached.
+				if (!Info.ShouldAttachToBaseComponent())
+				{
+					Info.Instances[InstanceIdx].BaseId = FFoliageInstanceBaseCache::InvalidBaseId;
+					Info.Instances[InstanceIdx].BaseComponent = nullptr;
+				}
 				Info.AddToBaseHash(InstanceIdx);
 				Info.InstanceHash->InsertInstance(Info.Instances[InstanceIdx].Location, InstanceIdx);
 			}
@@ -3746,6 +3807,16 @@ void AInstancedFoliageActor::NotifyFoliageTypeChanged(UFoliageType* FoliageType,
 	}
 }
 
+void AInstancedFoliageActor::RemoveBaseComponentOnFoliageTypeInstances(UFoliageType* FoliageType)
+{
+	FFoliageInfo* TypeInfo = FindInfo(FoliageType);
+
+	if (TypeInfo)
+	{
+		TypeInfo->RemoveBaseComponentOnInstances();
+	}
+}
+
 void AInstancedFoliageActor::NotifyFoliageTypeWillChange(UFoliageType* FoliageType)
 {
 	FFoliageInfo* TypeInfo = FindInfo(FoliageType);
@@ -3770,6 +3841,11 @@ void AInstancedFoliageActor::OnLevelActorMoved(AActor* InActor)
 				MoveInstancesForMovedComponent(Component);
 			}
 		}
+
+		if (FFoliageHelper::IsOwnedByFoliage(InActor))
+		{
+			MoveInstancesForMovedOwnedActors(InActor);
+		}
 	}
 }
 
@@ -3786,6 +3862,37 @@ void AInstancedFoliageActor::OnLevelActorDeleted(AActor* InActor)
 				DeleteInstancesForComponent(Component);
 			}
 		}
+
+		// Cleanup Foliage Instances if Actor is deleted outside of Foliage Tool
+		if (FFoliageHelper::IsOwnedByFoliage(InActor))
+		{
+			for (auto& Pair : FoliageInfos)
+			{
+				FFoliageInfo& Info = *Pair.Value;
+				UFoliageType* FoliageType = Pair.Key;
+
+				if (Info.Type == EFoliageImplType::Actor)
+				{
+					if (FFoliageActor* FoliageActor = StaticCast<FFoliageActor*>(Info.Implementation.Get()))
+					{
+						TArray<int32> InstancesToRemove;
+						// Make sure we find Null pointers and Delete those instances if we can't find the actor
+						FoliageActor->GetInvalidInstances(InstancesToRemove);
+						// Find Actor
+						int32 Index = FoliageActor->FindIndex(InActor);
+						if (Index != INDEX_NONE)
+						{
+							InstancesToRemove.Add(Index);
+						}
+						
+						if(InstancesToRemove.Num() > 0)
+						{
+							Info.RemoveInstances(this, InstancesToRemove, false);
+						}
+					}
+				}
+			}
+		}
 	}
 }
 
diff --git a/Engine/Source/Runtime/Foliage/Public/FoliageType_Actor.h b/Engine/Source/Runtime/Foliage/Public/FoliageType_Actor.h
--- a/Engine/Source/Runtime/Foliage/Public/FoliageType_Actor.h
+++ b/Engine/Source/Runtime/Foliage/Public/FoliageType_Actor.h
@@ -13,6 +13,9 @@ class UFoliageType_Actor : public UFoliageType
 	UPROPERTY(EditAnywhere, Category = Actor)
 	TSubclassOf<AActor> ActorClass;
 			
+	UPROPERTY(EditAnywhere, Category = Actor)
+	bool bShouldAttachToBaseComponent;
+
 	virtual UObject* GetSource() const override { return ActorClass; }
 
 #if WITH_EDITOR
diff --git a/Engine/Source/Runtime/Foliage/Public/InstancedFoliage.h b/Engine/Source/Runtime/Foliage/Public/InstancedFoliage.h
--- a/Engine/Source/Runtime/Foliage/Public/InstancedFoliage.h
+++ b/Engine/Source/Runtime/Foliage/Public/InstancedFoliage.h
@@ -206,6 +206,7 @@ struct FFoliageImpl
 	virtual void NotifyFoliageTypeChanged(AInstancedFoliageActor* IFA, UFoliageType* FoliageType, const TArray<FFoliageInstance>& Instances, const TSet<int32>& SelectedIndices, bool bSourceChanged) = 0;
 	virtual void EnterEditMode() {}
 	virtual void ExitEditMode() {}
+	virtual bool ShouldAttachToBaseComponent() const { return true; }
 #endif
 
 	virtual int32 GetOverlappingSphereCount(const FSphere& Sphere) const { return 0; }
@@ -237,6 +238,9 @@ struct FFoliageInfo
 
 	// Transient, editor-only list of selected instances.
 	TSet<int32> SelectedIndices;
+
+	// Moving instances
+	bool bMovingInstances;
 #endif
 
 	FOLIAGE_API FFoliageInfo();
@@ -299,6 +303,7 @@ struct FFoliageInfo
 
 	FOLIAGE_API void AddToBaseHash(int32 InstanceIdx);
 	FOLIAGE_API void RemoveFromBaseHash(int32 InstanceIdx);
+	FOLIAGE_API bool ShouldAttachToBaseComponent() const { return Implementation->ShouldAttachToBaseComponent(); }
 
 	// For debugging. Validate state after editing.
 	void CheckValid();
@@ -309,6 +314,8 @@ struct FFoliageInfo
 	FOLIAGE_API void PreEditUndo(AInstancedFoliageActor* IFA, UFoliageType* FoliageType);
 	FOLIAGE_API void EnterEditMode();
 	FOLIAGE_API void ExitEditMode();
+
+	FOLIAGE_API void RemoveBaseComponentOnInstances();
 #endif
 
 	friend FArchive& operator<<(FArchive& Ar, FFoliageInfo& MeshInfo);
@@ -357,12 +364,19 @@ public:
 		CellMap.FindOrAdd(Key).Add(InstanceIndex);
 	}
 
-	void RemoveInstance(const FVector& InstanceLocation, int32 InstanceIndex)
+	void RemoveInstance(const FVector& InstanceLocation, int32 InstanceIndex, bool bChecked = true)
 	{
 		uint64 Key = MakeKey(InstanceLocation);
 
-		int32 RemoveCount = CellMap.FindChecked(Key).Remove(InstanceIndex);
-		check(RemoveCount == 1);
+		if (bChecked)
+		{
+			int32 RemoveCount = CellMap.FindChecked(Key).Remove(InstanceIndex);
+			check(RemoveCount == 1);
+		}
+		else if(TSet<int32>* Value = CellMap.Find(Key))
+		{
+			Value->Remove(InstanceIndex);
+		}
 	}
 
 	void GetInstancesOverlappingBox(const FBox& InBox, TArray<int32>& OutInstanceIndices) const
diff --git a/Engine/Source/Runtime/Foliage/Public/InstancedFoliageActor.h b/Engine/Source/Runtime/Foliage/Public/InstancedFoliageActor.h
--- a/Engine/Source/Runtime/Foliage/Public/InstancedFoliageActor.h
+++ b/Engine/Source/Runtime/Foliage/Public/InstancedFoliageActor.h
@@ -179,6 +179,9 @@ public:
 	
 	// Move instances based on a component that has just been moved.
 	void MoveInstancesForMovedComponent(UActorComponent* InComponent);
+
+	// Move instances that are owned by foliage actor.
+	void MoveInstancesForMovedOwnedActors(AActor* InActor);
 	
 	// Returns a map of Static Meshes and their placed instances attached to a component.
 	FOLIAGE_API TMap<UFoliageType*, TArray<const FFoliageInstancePlacementInfo*>> GetInstancesForComponent(UActorComponent* InComponent);
@@ -240,6 +243,8 @@ public:
 
 	/* Fix up a duplicate IFA */
 	void RepairDuplicateIFA(AInstancedFoliageActor* InDuplicateIFA);
+
+	void RemoveBaseComponentOnFoliageTypeInstances(UFoliageType* FoliageType);
 #endif	//WITH_EDITOR
 
 private:
