Foliage: Fix Crash when deleting Actor FoliageType from palette (when there are instances in the map)

#rnx

diff --git a/Engine/Source/Runtime/Foliage/Private/FoliageActor.cpp b/Engine/Source/Runtime/Foliage/Private/FoliageActor.cpp
--- a/Engine/Source/Runtime/Foliage/Private/FoliageActor.cpp
+++ b/Engine/Source/Runtime/Foliage/Private/FoliageActor.cpp
@@ -38,7 +38,9 @@ void FFoliageActor::Serialize(FArchive& Ar)
 
 void FFoliageActor::DestroyActors(bool bOnLoad)
 {
-	for (AActor* Actor : ActorInstances)
+	TArray<AActor*> CopyActorInstances(ActorInstances);
+	ActorInstances.Empty();
+	for (AActor* Actor : CopyActorInstances)
 	{
 		if (Actor != nullptr)
 		{
@@ -49,7 +51,6 @@ void FFoliageActor::DestroyActors(bool bOnLoad)
 			Actor->GetWorld()->DestroyActor(Actor);
 		}
 	}
-	ActorInstances.Empty();
 }
 
 #if WITH_EDITOR
