Foliage: Fix Crash when deleting Foliage Asset

diff --git a/Engine/Source/Runtime/Foliage/Private/InstancedFoliage.cpp b/Engine/Source/Runtime/Foliage/Private/InstancedFoliage.cpp
--- a/Engine/Source/Runtime/Foliage/Private/InstancedFoliage.cpp
+++ b/Engine/Source/Runtime/Foliage/Private/InstancedFoliage.cpp
@@ -3152,6 +3152,7 @@ void AInstancedFoliageActor::Destroyed()
 				FoliageActor->DestroyActors(false);
 			}
 		}
+		FoliageInfos.Empty();
 	}
 
 	Super::Destroyed();
