Foliage: Fix DrawScale not being copied back to Foliage instance when changed outside of Foliage Tool

diff --git a/Engine/Source/Runtime/Foliage/Private/FoliageActor.cpp b/Engine/Source/Runtime/Foliage/Private/FoliageActor.cpp
--- a/Engine/Source/Runtime/Foliage/Private/FoliageActor.cpp
+++ b/Engine/Source/Runtime/Foliage/Private/FoliageActor.cpp
@@ -296,6 +296,8 @@ bool FFoliageActor::UpdateInstanceFromActor(AInstancedFoliageActor* IFA, AActor*
 	FTransform ActorTransform = InActor->GetTransform();
 	FoliageInfo.Instances[Index].Location = ActorTransform.GetLocation();
 	FoliageInfo.Instances[Index].Rotation = FRotator(ActorTransform.GetRotation());
+	FoliageInfo.Instances[Index].PreAlignRotation = FoliageInfo.Instances[Index].Rotation;
+	FoliageInfo.Instances[Index].DrawScale3D = InActor->GetActorScale3D();
 	FoliageInfo.InstanceHash->InsertInstance(FoliageInfo.Instances[Index].Location, Index);
 	
 	return true;
