Fix Actors being disconnected from Foliage when Compiling BP

diff --git a/Engine/Source/Runtime/Foliage/Private/InstancedFoliage.cpp b/Engine/Source/Runtime/Foliage/Private/InstancedFoliage.cpp
--- a/Engine/Source/Runtime/Foliage/Private/InstancedFoliage.cpp
+++ b/Engine/Source/Runtime/Foliage/Private/InstancedFoliage.cpp
@@ -3927,6 +3927,11 @@ void AInstancedFoliageActor::OnLevelActorDeleted(AActor* InActor)
 {
 	UWorld* InWorld = InActor->GetWorld();
 
+	if (GIsReinstancing)
+	{
+		return;
+	}
+
 	if (!InWorld || !InWorld->IsGameWorld())
 	{
 		for (UActorComponent* Component : InActor->GetComponents())
