Editing LevelPartitionGrid: Edigrate from //dev-editor-openworld

Multiple engine changes needed to support LevelGrid

#rnx

diff --git a/Engine/Source/Developer/AssetTools/Private/AssetRenameManager.cpp b/Engine/Source/Developer/AssetTools/Private/AssetRenameManager.cpp
--- a/Engine/Source/Developer/AssetTools/Private/AssetRenameManager.cpp
+++ b/Engine/Source/Developer/AssetTools/Private/AssetRenameManager.cpp
@@ -277,8 +277,33 @@ void FAssetRenameManager::FindSoftReferencesToObject(FSoftObjectPath TargetObjec
 	PopulateAssetReferencers(AssetsToRename);
 
 	// Load all referencing objects and find for referencing objects
-	TArray<UPackage*> ReferencingPackagesToSave;
-	LoadReferencingPackages(AssetsToRename, true, false, ReferencingPackagesToSave, ReferencingObjects);
+	TMap<FSoftObjectPath, TArray<UObject*>> ReferencingObjectsMap;
+
+	GatherReferencingObjects(AssetsToRename, ReferencingObjectsMap);
+
+	// Build an array out of the map results.
+	for (const auto& It : ReferencingObjectsMap)
+	{
+		for (UObject* Obj : It.Value)
+		{
+			ReferencingObjects.AddUnique(Obj);
+		}
+	}
+}
+
+void FAssetRenameManager::FindSoftReferencesToObjects(const TArray<FSoftObjectPath>& TargetObjects, TMap<FSoftObjectPath, TArray<UObject*>>& ReferencingObjects) const
+{
+	TArray<FAssetRenameDataWithReferencers> AssetsToRename;
+	for (const FSoftObjectPath& TargetObject : TargetObjects)
+	{
+		AssetsToRename.Emplace(FAssetRenameDataWithReferencers(FAssetRenameData(TargetObject, TargetObject, true)));
+	}
+
+	// Fill out referencers from asset registry
+	PopulateAssetReferencers(AssetsToRename);
+
+	// Load all referencing objects and find for referencing objects
+	GatherReferencingObjects(AssetsToRename, ReferencingObjects);
 }
 
 void FAssetRenameManager::FixReferencesAndRenameCallback(TArray<FAssetRenameData> AssetsAndNames, bool bAutoCheckout, bool bWithDialog) const
@@ -505,6 +530,13 @@ void FAssetRenameManager::PopulateAssetReferencers(TArray<FAssetRenameDataWithRe
 		}
 	}
 
+	TMap<FName, TArray<FName>> SoftReferencers;
+	TMap<FName, TArray<FName>> PackageReferencers;
+
+	TArray<UPackage*> ExtraPackagesToCheckForSoftReferences;
+	FEditorFileUtils::GetDirtyWorldPackages(ExtraPackagesToCheckForSoftReferences);
+	FEditorFileUtils::GetDirtyContentPackages(ExtraPackagesToCheckForSoftReferences);
+
 	// Gather all referencing packages for all assets that are being renamed
 	for (FAssetRenameDataWithReferencers& AssetToRename : AssetsToPopulate)
 	{
@@ -512,10 +544,14 @@ void FAssetRenameManager::PopulateAssetReferencers(TArray<FAssetRenameDataWithRe
 
 		FName OldPackageName = FName(*AssetToRename.OldObjectPath.GetLongPackageName());
 
-		TArray<FName> Referencers;
-		AssetRegistryModule.Get().GetReferencers(OldPackageName, Referencers, AssetToRename.bOnlyFixSoftReferences ? EAssetRegistryDependencyType::Soft : EAssetRegistryDependencyType::Packages);
+		TMap<FName, TArray<FName>>& ReferencersMap = AssetToRename.bOnlyFixSoftReferences ? SoftReferencers : PackageReferencers;
+		if (!ReferencersMap.Contains(OldPackageName))
+		{
+			TArray<FName>& Referencers = ReferencersMap.Add(OldPackageName);
+			AssetRegistryModule.Get().GetReferencers(OldPackageName, Referencers, AssetToRename.bOnlyFixSoftReferences ? EAssetRegistryDependencyType::Soft : EAssetRegistryDependencyType::Packages);
+		}
 
-		for (const FName& ReferencingPackageName : Referencers)
+		for (const FName& ReferencingPackageName : ReferencersMap.FindChecked(OldPackageName))
 		{
 			if (!RenamingAssetPackageNames.Contains(ReferencingPackageName))
 			{
@@ -529,11 +565,6 @@ void FAssetRenameManager::PopulateAssetReferencers(TArray<FAssetRenameDataWithRe
 			AssetToRename.ReferencingPackageNames.AddUnique(FName(*AssetToRename.NewObjectPath.GetLongPackageName()));
 
 			// Add dirty packages and the package that owns the reference. They will get filtered out in LoadReferencingPackages if they aren't valid
-			TArray<UPackage*> ExtraPackagesToCheckForSoftReferences;
-
-			FEditorFileUtils::GetDirtyWorldPackages(ExtraPackagesToCheckForSoftReferences);
-			FEditorFileUtils::GetDirtyContentPackages(ExtraPackagesToCheckForSoftReferences);
-
 			for (UPackage* Package : ExtraPackagesToCheckForSoftReferences)
 			{
 				AssetToRename.ReferencingPackageNames.AddUnique(Package->GetFName());
@@ -708,6 +739,49 @@ void FAssetRenameManager::LoadReferencingPackages(TArray<FAssetRenameDataWithRef
 	}
 }
 
+void FAssetRenameManager::GatherReferencingObjects(TArray<FAssetRenameDataWithReferencers>& AssetsToRename, TMap<FSoftObjectPath, TArray<UObject*>>& OutSoftReferencingObjects) const
+{
+	const UBlueprintEditorProjectSettings* EditorProjectSettings = GetDefault<UBlueprintEditorProjectSettings>();
+	bool bLoadPackagesForSoftReferences = EditorProjectSettings->bValidateUnloadedSoftActorReferences;
+
+	TMap<UPackage*, TMap<FSoftObjectPath, FSoftObjectPath>> ReferencingPackages;
+
+	for (int32 AssetIdx = 0; AssetIdx < AssetsToRename.Num(); ++AssetIdx)
+	{
+		FAssetRenameDataWithReferencers& RenameData = AssetsToRename[AssetIdx];
+
+		UObject* Asset = RenameData.Asset.Get();
+		if (!Asset)
+		{
+			// The asset for this rename must have been GCed or is otherwise invalid. Skip it unless this is a soft reference only fix
+			continue;
+		}
+
+		for (FName PackageName : RenameData.ReferencingPackageNames)
+		{
+			UPackage* Package = FindPackage(nullptr, *PackageName.ToString());
+
+			// Don't load package if this is a soft reference fix and the project settings say not to
+			if (!Package && (!RenameData.bOnlyFixSoftReferences || bLoadPackagesForSoftReferences))
+			{
+				Package = LoadPackage(nullptr, *PackageName.ToString(), LOAD_None);
+			}
+
+			if (Package)
+			{
+				ReferencingPackages.FindOrAdd(Package).Add(RenameData.OldObjectPath, RenameData.NewObjectPath);
+			}
+		}
+	}
+
+	TArray<UPackage*> PackagesToSaveForThisAsset;
+	bool bAllPackagesLoadedForThisAsset = true;
+	for (const auto& ReferencingPackage : ReferencingPackages)
+	{
+		CheckPackageForSoftObjectReferences(ReferencingPackage.Key, ReferencingPackage.Value, OutSoftReferencingObjects);
+	}
+}
+
 bool FAssetRenameManager::CheckOutPackages(TArray<FAssetRenameDataWithReferencers>& AssetsToRename, TArray<UPackage*>& InOutReferencingPackagesToSave, bool bAutoCheckout) const
 {
 	bool bUserAcceptedCheckout = true;
@@ -1077,6 +1151,23 @@ void FAssetRenameManager::OnMarkPackageDirty(UPackage* Pkg, bool bWasDirty)
 
 bool FAssetRenameManager::CheckPackageForSoftObjectReferences(UPackage* Package, const TMap<FSoftObjectPath, FSoftObjectPath>& AssetRedirectorMap, TArray<UObject*>& OutReferencingObjects) const
 {	
+	TMap<FSoftObjectPath, TArray<UObject*>> ReferencingObjectsMap;
+	
+	CheckPackageForSoftObjectReferences(Package, AssetRedirectorMap, ReferencingObjectsMap);
+
+	// Build an array out of the map results.
+	for (const auto& It : ReferencingObjectsMap)
+	{
+		for (UObject* Obj : It.Value)
+		{
+			OutReferencingObjects.AddUnique(Obj);
+		}
+	}
+	return OutReferencingObjects.Num() != 0;
+}
+
+bool FAssetRenameManager::CheckPackageForSoftObjectReferences(UPackage* Package, const TMap<FSoftObjectPath, FSoftObjectPath>& AssetRedirectorMap, TMap<FSoftObjectPath, TArray<UObject*>>& OutReferencingObjects) const
+{
 	using namespace AssetRenameManagerImpl;
 
 	struct FSoftObjectPathLess
@@ -1089,7 +1180,7 @@ bool FAssetRenameManager::CheckPackageForSoftObjectReferences(UPackage* Package,
 	};
 
 	bool bFoundReference = false;
-	
+
 	// First check cache
 	FCachedSoftReference* CachedReferences = CachedSoftReferences.Find(Package->GetFName());
 
@@ -1160,7 +1251,7 @@ bool FAssetRenameManager::CheckPackageForSoftObjectReferences(UPackage* Package,
 					UObject* ObjectPtr = WeakPtr.Get();
 					if (ObjectPtr)
 					{
-						OutReferencingObjects.AddUnique(ObjectPtr);
+						OutReferencingObjects.FindOrAdd(CachedKey).AddUnique(ObjectPtr);
 					}
 				}
 			}
diff --git a/Engine/Source/Developer/AssetTools/Private/AssetRenameManager.h b/Engine/Source/Developer/AssetTools/Private/AssetRenameManager.h
--- a/Engine/Source/Developer/AssetTools/Private/AssetRenameManager.h
+++ b/Engine/Source/Developer/AssetTools/Private/AssetRenameManager.h
@@ -38,6 +38,9 @@ public:
 	/** Returns list of objects that soft reference the given soft object path. This will load assets into memory to verify */
 	void FindSoftReferencesToObject(FSoftObjectPath TargetObject, TArray<UObject*>& ReferencingObjects) const;
 
+	/** Returns list of objects that soft reference the given soft object paths. This will load assets into memory to verify */
+	void FindSoftReferencesToObjects(const TArray<FSoftObjectPath>& TargetObjects, TMap<FSoftObjectPath, TArray<UObject*>>& ReferencingObjects) const;
+
 	/** Accessor for post rename event */
 	FAssetPostRenameEvent& OnAssetPostRenameEvent() { return AssetPostRenameEvent; }
 
@@ -52,6 +55,9 @@ public:
 	/** Filters packages list depending on if it actually has soft object paths pointing to the specific object being renamed */
 	bool CheckPackageForSoftObjectReferences(UPackage* Package, const TMap<FSoftObjectPath, FSoftObjectPath>& AssetRedirectorMap, TArray<UObject*>& OutReferencingObjects) const;
 
+	/** Filters packages list depending on if it actually has soft object paths pointing to the specific object being renamed */
+	bool CheckPackageForSoftObjectReferences(UPackage* Package, const TMap<FSoftObjectPath, FSoftObjectPath>& AssetRedirectorMap, TMap<FSoftObjectPath, TArray<UObject*>>& OutReferencingObjects) const;
+
 private:
 	/** Callback used by DiscoverintAssetsDialog to call FixrefrencesAndRename */
 	void FixReferencesAndRenameCallback(TArray<FAssetRenameData> AssetsAndNames, bool bAutoCheckout, bool bWithDialog) const;
@@ -76,6 +82,9 @@ private:
 	 */
 	void LoadReferencingPackages(TArray<FAssetRenameDataWithReferencers>& AssetsToRename, bool bLoadAllPackages, bool bCheckStatus, TArray<UPackage*>& OutReferencingPackagesToSave, TArray<UObject*>& OutSoftReferencingObjects) const;
 
+    /** Gather a list of referencing object for each of the asset in AssetsToRename. Will load all referencing packages */
+	void GatherReferencingObjects(TArray<FAssetRenameDataWithReferencers>& AssetsToRename, TMap<FSoftObjectPath, TArray<UObject*>>& OutSoftReferencingObjects) const;
+
 	/** 
 	 * Prompts to check out the source package and all referencing packages and marks assets whose referencing packages were not checked out to leave a redirector.
 	 * Trims PackagesToSave when necessary.
diff --git a/Engine/Source/Developer/AssetTools/Private/AssetTools.cpp b/Engine/Source/Developer/AssetTools/Private/AssetTools.cpp
--- a/Engine/Source/Developer/AssetTools/Private/AssetTools.cpp
+++ b/Engine/Source/Developer/AssetTools/Private/AssetTools.cpp
@@ -1006,6 +1006,11 @@ void UAssetToolsImpl::FindSoftReferencesToObject(FSoftObjectPath TargetObject, T
 	AssetRenameManager->FindSoftReferencesToObject(TargetObject, ReferencingObjects);
 }
 
+void UAssetToolsImpl::FindSoftReferencesToObjects(const TArray<FSoftObjectPath>& TargetObjects, TMap<FSoftObjectPath, TArray<UObject*>>& ReferencingObjects)
+{
+	AssetRenameManager->FindSoftReferencesToObjects(TargetObjects, ReferencingObjects);
+}
+
 void UAssetToolsImpl::RenameReferencingSoftObjectPaths(const TArray<UPackage *> PackagesToCheck, const TMap<FSoftObjectPath, FSoftObjectPath>& AssetRedirectorMap)
 {
 	AssetRenameManager->RenameReferencingSoftObjectPaths(PackagesToCheck, AssetRedirectorMap);
diff --git a/Engine/Source/Developer/AssetTools/Private/AssetTools.h b/Engine/Source/Developer/AssetTools/Private/AssetTools.h
--- a/Engine/Source/Developer/AssetTools/Private/AssetTools.h
+++ b/Engine/Source/Developer/AssetTools/Private/AssetTools.h
@@ -81,6 +81,7 @@ public:
 	virtual bool RenameAssets(const TArray<FAssetRenameData>& AssetsAndNames) override;
 	virtual void RenameAssetsWithDialog(const TArray<FAssetRenameData>& AssetsAndNames, bool bAutoCheckout = false) override;
 	virtual void FindSoftReferencesToObject(FSoftObjectPath TargetObject, TArray<UObject*>& ReferencingObjects) override;
+	virtual void FindSoftReferencesToObjects(const TArray<FSoftObjectPath>& TargetObjects, TMap<FSoftObjectPath, TArray<UObject*>>& ReferencingObjects) override;
 	virtual void RenameReferencingSoftObjectPaths(const TArray<UPackage *> PackagesToCheck, const TMap<FSoftObjectPath, FSoftObjectPath>& AssetRedirectorMap) override;
 	virtual TArray<UObject*> ImportAssets(const FString& DestinationPath) override;
 	virtual TArray<UObject*> ImportAssetsWithDialog(const FString& DestinationPath) override;
diff --git a/Engine/Source/Developer/AssetTools/Public/IAssetTools.h b/Engine/Source/Developer/AssetTools/Public/IAssetTools.h
--- a/Engine/Source/Developer/AssetTools/Public/IAssetTools.h
+++ b/Engine/Source/Developer/AssetTools/Public/IAssetTools.h
@@ -236,6 +236,9 @@ public:
 	UFUNCTION(BlueprintCallable, Category = "Editor Scripting | Asset Tools")
 	virtual void FindSoftReferencesToObject(FSoftObjectPath TargetObject, TArray<UObject*>& ReferencingObjects) = 0;
 
+	/** Returns list of objects that soft reference the given soft object paths. This will load assets into memory to verify */
+	virtual void FindSoftReferencesToObjects(const TArray<FSoftObjectPath>& TargetObjects, TMap<FSoftObjectPath, TArray<UObject*>>& ReferencingObjects) = 0;
+
 	/**
 	 * Function that renames all FSoftObjectPath object with the old asset path to the new one.
 	 *
diff --git a/Engine/Source/Editor/FoliageEdit/Private/FoliageEdMode.cpp b/Engine/Source/Editor/FoliageEdit/Private/FoliageEdMode.cpp
--- a/Engine/Source/Editor/FoliageEdit/Private/FoliageEdMode.cpp
+++ b/Engine/Source/Editor/FoliageEdit/Private/FoliageEdMode.cpp
@@ -232,6 +232,7 @@ FEdModeFoliage::FEdModeFoliage()
 	, FoliageInteractor(nullptr)
 	, UpdateSelectionCounter(0)
 	, bHasDeferredSelectionNotification(false)
+	, bMoving(false)
 {
 	// Load resources and construct brush component
 	UStaticMesh* StaticMesh = nullptr;
@@ -1508,22 +1509,32 @@ static void SpawnFoliageInstance(UWorld* InWorld, const UFoliageType* Settings,
 {
 	SCOPE_CYCLE_COUNTER(STAT_FoliageSpawnInstance);
 
-	TMap<ULevel*, TSet<const FFoliageInstance*>> PerLevelPlacedInstances;
+	TMap<ULevel*, TArray<const FFoliageInstance*>> PerLevelPlacedInstances;
 
 	if (UISettings != nullptr && UISettings->GetIsInSpawnInCurrentLevelMode())
 	{
-		TSet<const FFoliageInstance*>& LevelInstances = PerLevelPlacedInstances.FindOrAdd(InWorld->GetCurrentLevel());
-
-		for (const FFoliageInstance& PlacedInstance : PlacedInstances)
+		if (const ILevelPartitionInterface* LevelPartition = InWorld->GetCurrentLevel()->GetLevelPartition())
 		{
-			LevelInstances.Add(&PlacedInstance);
+			for (const FFoliageInstance& PlacedInstance : PlacedInstances)
+			{
+				PerLevelPlacedInstances.FindOrAdd(LevelPartition->GetSubLevel(PlacedInstance.Location)).Add(&PlacedInstance);
+			}
+		}
+		else
+		{
+			TArray<const FFoliageInstance*>& LevelInstances = PerLevelPlacedInstances.FindOrAdd(InWorld->GetCurrentLevel());
+
+			for (const FFoliageInstance& PlacedInstance : PlacedInstances)
+			{
+				LevelInstances.Add(&PlacedInstance);
+			}
 		}
 	}
 	else
 	{
 		for (const FFoliageInstance& PlacedInstance : PlacedInstances)
 		{
-			TSet<const FFoliageInstance*>& LevelInstances = PerLevelPlacedInstances.FindOrAdd(PlacedInstance.BaseComponent->GetComponentLevel());
+			TArray<const FFoliageInstance*>& LevelInstances = PerLevelPlacedInstances.FindOrAdd(PlacedInstance.BaseComponent->GetComponentLevel());
 			LevelInstances.Add(&PlacedInstance);
 		}
 	}
@@ -2010,8 +2021,72 @@ void FEdModeFoliage::ApplySelection(UWorld* InWorld, bool bApply)
 	}
 }
 
-void FEdModeFoliage::TransformSelectedInstances(UWorld* InWorld, const FVector& InDrag, const FRotator& InRot, const FVector& InScale, bool bDuplicate)
+void FEdModeFoliage::UpdateInstancePartitioning(UWorld* InWorld)
+{
+	if (!bMoving)
+	{
+		return;
+	}
+
+	const int32 NumLevels = InWorld->GetNumLevels();
+	for (int32 LevelIdx = 0; LevelIdx < NumLevels; ++LevelIdx)
+	{
+		ULevel* Level = InWorld->GetLevel(LevelIdx);
+
+		const ILevelPartitionInterface* LevelPartition = Level->GetLevelPartition();
+		if (LevelPartition == nullptr)
+		{
+			continue;
+		}
+		
+		AInstancedFoliageActor* IFA = AInstancedFoliageActor::GetInstancedFoliageActorForLevel(Level);
+		if (IFA)
+		{
+			for (auto& MeshPair : IFA->FoliageInfos)
+			{
+				FFoliageInfo& FoliageInfo = *MeshPair.Value;
+				if (FoliageInfo.Type == EFoliageImplType::Actor)
+				{
+					continue; // Actors are handled through the Partitioning code
+				}
+
+				ULevel* TargetLevel = nullptr;
+				// Loop here because SelectedIndices will change on MoveInstancesToLevel so we need to process the modified remaining SelectedIndices
+				do
+				{
+					TargetLevel = nullptr;
+					TSet<int32> InstancesToMove;
+
+					for (int32 SelectedInstanceIdx : FoliageInfo.SelectedIndices)
+					{
+						FFoliageInstance& Instance = FoliageInfo.Instances[SelectedInstanceIdx];
+						ULevel* NewLevel = LevelPartition->GetSubLevel(Instance.Location);
+						if ((TargetLevel == nullptr || TargetLevel == NewLevel) && NewLevel != Level)
+						{
+							TargetLevel = NewLevel;
+							InstancesToMove.Add(SelectedInstanceIdx);
+						}
+					}
+
+					if (TargetLevel && InstancesToMove.Num())
+					{
+						IFA->MoveInstancesToLevel(TargetLevel, InstancesToMove, &MeshPair.Value.Get(), MeshPair.Key, true);
+					}
+
+				} while (TargetLevel && FoliageInfo.SelectedIndices.Num() > 0);
+			}
+		}
+	}
+}
+
+void FEdModeFoliage::PostTransformSelectedInstances(UWorld* InWorld)
 {
+	if (!bMoving)
+	{
+		return;
+	}
+
+	bMoving = false;
 	const int32 NumLevels = InWorld->GetNumLevels();
 	for (int32 LevelIdx = 0; LevelIdx < NumLevels; ++LevelIdx)
 	{
@@ -2026,6 +2101,32 @@ void FEdModeFoliage::TransformSelectedInstances(UWorld* InWorld, const FVector&
 				FFoliageInfo& FoliageInfo = *MeshPair.Value;
 				TArray<int32> SelectedIndices = FoliageInfo.SelectedIndices.Array();
 
+				if (SelectedIndices.Num() > 0)
+				{
+					const bool bFinished = true;
+					FoliageInfo.PostMoveInstances(IFA, SelectedIndices, bFinished);
+				}
+			}
+		}
+	}
+}
+
+void FEdModeFoliage::TransformSelectedInstances(UWorld* InWorld, const FVector& InDrag, const FRotator& InRot, const FVector& InScale, bool bDuplicate)
+{
+	const int32 NumLevels = InWorld->GetNumLevels();
+	for (int32 LevelIdx = 0; LevelIdx < NumLevels; ++LevelIdx)
+	{
+		ULevel* Level = InWorld->GetLevel(LevelIdx);
+		AInstancedFoliageActor* IFA = AInstancedFoliageActor::GetInstancedFoliageActorForLevel(Level);
+		if (IFA)
+		{
+			bool bFoundSelection = false;
+
+			for (auto& MeshPair : IFA->FoliageInfos)
+			{
+				FFoliageInfo& FoliageInfo = *MeshPair.Value;
+				TArray<int32> SelectedIndices = FoliageInfo.SelectedIndices.Array();
+			
 				if (SelectedIndices.Num() > 0)
 				{
 					// Mark actor once we found selection
@@ -2041,6 +2142,7 @@ void FEdModeFoliage::TransformSelectedInstances(UWorld* InWorld, const FVector&
 						OnInstanceCountUpdated(MeshPair.Key);
 					}
 
+					bMoving = true;
 					FoliageInfo.PreMoveInstances(IFA, SelectedIndices);
 
 					for (int32 SelectedInstanceIdx : SelectedIndices)
@@ -2052,7 +2154,8 @@ void FEdModeFoliage::TransformSelectedInstances(UWorld* InWorld, const FVector&
 						Instance.DrawScale3D += InScale;
 					}
 
-					FoliageInfo.PostMoveInstances(IFA, SelectedIndices);
+					const bool bFinished = false;
+					FoliageInfo.PostMoveInstances(IFA, SelectedIndices, bFinished);
 				}
 			}
 
@@ -2064,44 +2167,34 @@ void FEdModeFoliage::TransformSelectedInstances(UWorld* InWorld, const FVector&
 	}
 }
 
-AInstancedFoliageActor* FEdModeFoliage::GetSelectionLocation(UWorld* InWorld, FVector& OutLocation) const
+bool FEdModeFoliage::GetSelectionLocation(UWorld* InWorld, FVector& OutLocation) const
 {
-	// Prefer current level
-	{
-		AInstancedFoliageActor* IFA = AInstancedFoliageActor::GetInstancedFoliageActorForCurrentLevel(InWorld);
-		if (IFA && IFA->GetSelectionLocation(OutLocation))
-		{
-			return IFA;
-		}
-	}
-
+	FBox OutBox(EForceInit::ForceInit);
+	bool bHasSelection = false;
 	// Go through all sub-levels
 	const int32 NumLevels = InWorld->GetNumLevels();
 	for (int32 LevelIdx = 0; LevelIdx < NumLevels; ++LevelIdx)
 	{
 		ULevel* Level = InWorld->GetLevel(LevelIdx);
-		if (Level != InWorld->GetCurrentLevel())
-		{
-			AInstancedFoliageActor* IFA = AInstancedFoliageActor::GetInstancedFoliageActorForLevel(Level);
-			if (IFA && IFA->GetSelectionLocation(OutLocation))
-			{
-				return IFA;
-			}
-		}
+		AInstancedFoliageActor* IFA = AInstancedFoliageActor::GetInstancedFoliageActorForLevel(Level);
+		bHasSelection |= (IFA && IFA->GetSelectionLocation(OutBox));
 	}
 
-	return nullptr;
+	if (bHasSelection)
+	{
+		OutLocation = OutBox.GetCenter();
+	}
+
+	return bHasSelection;
 }
 
 void FEdModeFoliage::UpdateWidgetLocationToInstanceSelection()
 {
 	FVector SelectionLocation = FVector::ZeroVector;
-	AInstancedFoliageActor* IFA = GetSelectionLocation(GetWorld(), SelectionLocation);
-	Owner->PivotLocation = Owner->SnappedLocation = SelectionLocation;
-	//if (IFA)
-	//{
-	//	IFA->MarkComponentsRenderStateDirty();
-	//}
+	if (GetSelectionLocation(GetWorld(), SelectionLocation))
+	{
+		Owner->PivotLocation = Owner->SnappedLocation = SelectionLocation;
+	}
 }
 
 void FEdModeFoliage::RemoveSelectedInstances(UWorld* InWorld)
@@ -3958,6 +4051,8 @@ bool FEdModeFoliage::EndTracking(FEditorViewportClient* InViewportClient, FViewp
 {
 	if (!bAdjustBrushRadius && (UISettings.GetSelectToolSelected() || UISettings.GetLassoSelectToolSelected()))
 	{
+		UpdateInstancePartitioning(GetWorld());
+		PostTransformSelectedInstances(GetWorld());
 		GEditor->EndTransaction();
 		return true;
 	}
@@ -4016,7 +4111,7 @@ bool FEdModeFoliage::ShouldDrawWidget() const
 	if (UISettings.GetSelectToolSelected() || (UISettings.GetLassoSelectToolSelected() && !bToolActive))
 	{
 		FVector Location;
-		return GetSelectionLocation(GetWorld(), Location) != nullptr;
+		return GetSelectionLocation(GetWorld(), Location);
 	}
 	return false;
 }
diff --git a/Engine/Source/Editor/FoliageEdit/Private/FoliageEdMode.h b/Engine/Source/Editor/FoliageEdit/Private/FoliageEdMode.h
--- a/Engine/Source/Editor/FoliageEdit/Private/FoliageEdMode.h
+++ b/Engine/Source/Editor/FoliageEdit/Private/FoliageEdMode.h
@@ -590,11 +590,17 @@ private:
 	/**  Propagate the selected foliage instances to the actual render foliage */
 	void ApplySelection(UWorld* InWorld, bool bApply);
 
+	/**  Update Instances so that they are in the right partitioning level */
+	void UpdateInstancePartitioning(UWorld* InWorld);
+
+	/** Called when transform transaction is done */
+	void PostTransformSelectedInstances(UWorld* InWorld);
+
 	/**  Applies relative transformation to selected instances */
 	void TransformSelectedInstances(UWorld* InWorld, const FVector& InDrag, const FRotator& InRot, const FVector& InScale, bool bDuplicate);
 
-	/**  Return selected actor and instance location */
-	AInstancedFoliageActor* GetSelectionLocation(UWorld* InWorld, FVector& OutLocation) const;
+	/**  Return true if OutLocation is valid */
+	bool GetSelectionLocation(UWorld* InWorld, FVector& OutLocation) const;
 	
 	/**  Updates ed mode widget location to currently selected instance */
 	void UpdateWidgetLocationToInstanceSelection();
@@ -679,5 +685,8 @@ private:
 	int32 UpdateSelectionCounter;
 	bool bHasDeferredSelectionNotification;
 	friend class FEdModeFoliageSelectionUpdate;
+
+	/** When transforming instances */
+	bool bMoving;
 };
 
diff --git a/Engine/Source/Editor/UnrealEd/Classes/Editor/UnrealEdEngine.h b/Engine/Source/Editor/UnrealEd/Classes/Editor/UnrealEdEngine.h
--- a/Engine/Source/Editor/UnrealEd/Classes/Editor/UnrealEdEngine.h
+++ b/Engine/Source/Editor/UnrealEd/Classes/Editor/UnrealEdEngine.h
@@ -840,6 +840,11 @@ public:
 	DECLARE_DELEGATE_RetVal_OneParam(TSharedPtr<IAssetReferenceFilter>, FOnMakeAssetReferenceFilter, const FAssetReferenceFilterContext& /*Context*/);
 	FOnMakeAssetReferenceFilter& OnMakeAssetReferenceFilter() { return OnMakeAssetReferenceFilterDelegate; }
 	TSharedPtr<IAssetReferenceFilter> MakeAssetReferenceFilter(const FAssetReferenceFilterContext& Context) { return OnMakeAssetReferenceFilterDelegate.IsBound() ? OnMakeAssetReferenceFilterDelegate.Execute(Context) : nullptr; }
+
+	/** Add Visualizers that are active even when Actor is not part of the selection. */
+	void AddNonSelectionVisualizers(AActor* Actor);
+	/** Remove Visualizers that are active even when Actor is not part of the selection */
+	void RemoveNonSelectionVisualizers(AActor* Actor);
 	
 protected:
 
@@ -860,7 +865,12 @@ protected:
 	/**
 	 * The list of visualizers to draw when selection changes
 	 */
-	TArray<FCachedComponentVisualizer> VisualizersForSelection;
+	TArray<FCachedComponentVisualizer> SelectionVisualizers;
+
+	/**
+	 * The list of visualizers to draw independant of selection
+	 */
+	TArray<FCachedComponentVisualizer> NonSelectionVisualizers;
 
 	/** Instance responsible for monitoring this editor's performance */
 	FPerformanceMonitor* PerformanceMonitor;
@@ -884,4 +894,7 @@ private:
 	* @return	Returns the number of dirty packages that require checkout. If bCheckIfAny is true, returns 1 if any packages will require checkout.
 	*/
 	int32 InternalGetNumDirtyPackagesThatNeedCheckout(bool bCheckIfAny) const;
+
+	/** Internal function to filter and add visualizers to a specific list */
+	void AddVisualizers(AActor* Actor, TArray<FCachedComponentVisualizer>& Visualizers, TFunctionRef<bool(const TSharedPtr<FComponentVisualizer>&)> Condition);
 };
diff --git a/Engine/Source/Editor/UnrealEd/Private/EditorLevelUtils.cpp b/Engine/Source/Editor/UnrealEd/Private/EditorLevelUtils.cpp
--- a/Engine/Source/Editor/UnrealEd/Private/EditorLevelUtils.cpp
+++ b/Engine/Source/Editor/UnrealEd/Private/EditorLevelUtils.cpp
@@ -314,7 +314,10 @@ ULevel* UEditorLevelUtils::AddLevelsToWorld(UWorld* InWorld, TArray<FString> Pac
 	FEditorDelegates::RefreshLevelBrowser.Broadcast();
 
 	// Update volume actor visibility for each viewport since we loaded a level which could potentially contain volumes
-	GUnrealEd->UpdateVolumeActorVisibility(nullptr);
+	if (GUnrealEd)
+	{
+		GUnrealEd->UpdateVolumeActorVisibility(nullptr);
+	}
 
 	return NewLevel;
 }
@@ -362,7 +365,10 @@ ULevelStreaming* UEditorLevelUtils::AddLevelToWorld(UWorld* InWorld, const TCHAR
 	FEditorDelegates::RefreshLevelBrowser.Broadcast();
 
 	// Update volume actor visibility for each viewport since we loaded a level which could potentially contain volumes
-	GUnrealEd->UpdateVolumeActorVisibility(nullptr);
+	if (GUnrealEd)
+	{
+		GUnrealEd->UpdateVolumeActorVisibility(nullptr);
+	}
 
 	return NewStreamingLevel;
 }
diff --git a/Engine/Source/Editor/UnrealEd/Private/LevelEditorViewport.cpp b/Engine/Source/Editor/UnrealEd/Private/LevelEditorViewport.cpp
--- a/Engine/Source/Editor/UnrealEd/Private/LevelEditorViewport.cpp
+++ b/Engine/Source/Editor/UnrealEd/Private/LevelEditorViewport.cpp
@@ -3057,7 +3057,9 @@ void FLevelEditorViewportClient::TrackingStopped()
 		// Create post edit property change event
 		UProperty* TransformProperty = LevelEditorViewportClientHelper::GetEditTransformProperty(GetWidgetMode());
 		FPropertyChangedEvent PropertyChangedEvent(TransformProperty, EPropertyChangeType::ValueSet);
-		
+
+		TArray<AActor*> MovedActors;
+
 		for (FSelectionIterator It(GEditor->GetSelectedActorIterator()); It; ++It) 
 		{
 			AActor* Actor = static_cast<AActor*>( *It );
@@ -3133,11 +3135,15 @@ void FLevelEditorViewportClient::TrackingStopped()
 				// Broadcast Post Edit change notification, we can't call PostEditChangeProperty directly on Actor or ActorComponent from here since it wasn't pair with a proper PreEditChange
 				FCoreUObjectDelegates::OnObjectPropertyChanged.Broadcast(Actor, PropertyChangedEvent);
 				Actor->PostEditMove(true);
+
+				MovedActors.Add(Actor);
 	
 				GEditor->BroadcastEndObjectMovement(*Actor);
 			}
 		}
 
+		GEditor->BroadcastActorsMoved(MovedActors);
+
 		if (!GUnrealEd->IsPivotMovedIndependently())
 		{
 			GUnrealEd->UpdatePivotLocationForSelection();
diff --git a/Engine/Source/Editor/UnrealEd/Private/UnrealEdEngine.cpp b/Engine/Source/Editor/UnrealEd/Private/UnrealEdEngine.cpp
--- a/Engine/Source/Editor/UnrealEd/Private/UnrealEdEngine.cpp
+++ b/Engine/Source/Editor/UnrealEd/Private/UnrealEdEngine.cpp
@@ -1300,7 +1300,7 @@ void UUnrealEdEngine::RegisterComponentVisualizer(FName ComponentClassName, TSha
 void UUnrealEdEngine::UnregisterComponentVisualizer(FName ComponentClassName)
 {
 	TSharedPtr<FComponentVisualizer> Visualizer = FindComponentVisualizer(ComponentClassName);
-	VisualizersForSelection.RemoveAll([&Visualizer](const auto& CachedComponentVisualizer) { return CachedComponentVisualizer.Visualizer == Visualizer; });
+	SelectionVisualizers.RemoveAll([&Visualizer](const auto& CachedComponentVisualizer) { return CachedComponentVisualizer.Visualizer == Visualizer; });
 
 	ComponentVisualizerMap.Remove(ComponentClassName);
 }
@@ -1335,7 +1335,12 @@ TSharedPtr<class FComponentVisualizer> UUnrealEdEngine::FindComponentVisualizer(
 
 void UUnrealEdEngine::DrawComponentVisualizers(const FSceneView* View, FPrimitiveDrawInterface* PDI)
 {
-	for(FCachedComponentVisualizer& CachedVisualizer : VisualizersForSelection)
+	for(FCachedComponentVisualizer& CachedVisualizer : SelectionVisualizers)
+	{
+		CachedVisualizer.Visualizer->DrawVisualization(CachedVisualizer.ComponentPropertyPath.GetComponent(), View, PDI);
+	}
+
+	for(FCachedComponentVisualizer& CachedVisualizer : NonSelectionVisualizers)
 	{
 		CachedVisualizer.Visualizer->DrawVisualization(CachedVisualizer.ComponentPropertyPath.GetComponent(), View, PDI);
 	}
@@ -1344,19 +1349,44 @@ void UUnrealEdEngine::DrawComponentVisualizers(const FSceneView* View, FPrimitiv
 
 void UUnrealEdEngine::DrawComponentVisualizersHUD(const FViewport* Viewport, const FSceneView* View, FCanvas* Canvas)
 {
-	for(FCachedComponentVisualizer& CachedVisualizer : VisualizersForSelection)
+	for(FCachedComponentVisualizer& CachedVisualizer : SelectionVisualizers)
+	{
+		CachedVisualizer.Visualizer->DrawVisualizationHUD(CachedVisualizer.ComponentPropertyPath.GetComponent(), Viewport, View, Canvas);
+	}
+
+	for(FCachedComponentVisualizer& CachedVisualizer : NonSelectionVisualizers)
 	{
 		CachedVisualizer.Visualizer->DrawVisualizationHUD(CachedVisualizer.ComponentPropertyPath.GetComponent(), Viewport, View, Canvas);
 	}
 }
 
+void UUnrealEdEngine::AddVisualizers(AActor* Actor, TArray<FCachedComponentVisualizer>& Visualizers, TFunctionRef<bool(const TSharedPtr<FComponentVisualizer>&)> Condition)
+{
+	TInlineComponentArray<UActorComponent*> Components;
+	Actor->GetComponents(Components, true);
+
+	for(int32 CompIdx = 0; CompIdx < Components.Num(); CompIdx++)
+	{
+		UActorComponent* Comp = Components[CompIdx];
+		if(Comp->IsRegistered())
+		{
+			// Try and find a visualizer
+			TSharedPtr<FComponentVisualizer> Visualizer = FindComponentVisualizer(Comp->GetClass());
+			if(Visualizer.IsValid() && Condition(Visualizer))
+			{
+				Visualizers.Add(FCachedComponentVisualizer(Comp, Visualizer));
+			}
+		}
+	}
+}
+
 void UUnrealEdEngine::OnEditorSelectionChanged(UObject* SelectionThatChanged)
 {
 	if(SelectionThatChanged == GetSelectedActors())
 	{
 		// actor selection changed.  Update the list of component visualizers
 		// This is expensive so we do not search for visualizers each time they want to draw
-		VisualizersForSelection.Empty();
+		SelectionVisualizers.Empty();
 
 		// Iterate over all selected actors
 		for(FSelectionIterator It(GetSelectedActorIterator()); It; ++It)
@@ -1364,28 +1394,36 @@ void UUnrealEdEngine::OnEditorSelectionChanged(UObject* SelectionThatChanged)
 			AActor* Actor = Cast<AActor>(*It);
 			if(Actor != nullptr)
 			{
-				// Then iterate over components of that actor (and recurse through child components)
-				TInlineComponentArray<UActorComponent*> Components;
-				Actor->GetComponents(Components, true);
-
-				for(int32 CompIdx = 0; CompIdx < Components.Num(); CompIdx++)
-				{
-					UActorComponent* Comp = Components[CompIdx];
-					if(Comp->IsRegistered())
-					{
-						// Try and find a visualizer
-						TSharedPtr<FComponentVisualizer> Visualizer = FindComponentVisualizer(Comp->GetClass());
-						if(Visualizer.IsValid())
-						{
-							VisualizersForSelection.Add(FCachedComponentVisualizer(Comp, Visualizer));
-						}
-					}
-				}
+				AddVisualizers(Actor, SelectionVisualizers, [](const TSharedPtr<FComponentVisualizer>& Visualizer) { return Visualizer->ShowWhenSelected(); });
 			}
 		}
 	}
 }
 
+void UUnrealEdEngine::AddNonSelectionVisualizers(AActor* Actor)
+{
+	AddVisualizers(Actor, NonSelectionVisualizers, [](const TSharedPtr<FComponentVisualizer>& Visualizer) { return !Visualizer->ShowWhenSelected(); });
+}
+
+void UUnrealEdEngine::RemoveNonSelectionVisualizers(AActor* Actor)
+{
+	TInlineComponentArray<UActorComponent*> Components;
+	Actor->GetComponents(Components);
+
+	for (int32 CompIdx = 0; CompIdx < Components.Num(); CompIdx++)
+	{
+		UActorComponent* Comp = Components[CompIdx];
+		if (Comp->IsRegistered())
+		{
+			// Try and find a visualizer
+			TSharedPtr<FComponentVisualizer> Visualizer = FindComponentVisualizer(Comp->GetClass());
+			if (Visualizer.IsValid())
+			{
+				NonSelectionVisualizers.RemoveAll([&Visualizer, Actor](const auto& CachedComponentVisualizer) { return CachedComponentVisualizer.ComponentPropertyPath.GetComponent() && CachedComponentVisualizer.ComponentPropertyPath.GetComponent()->GetOwner() == Actor; });
+			}
+		}
+	}
+}
 
 EWriteDisallowedWarningState UUnrealEdEngine::GetWarningStateForWritePermission(const FString& PackageName) const
 {
diff --git a/Engine/Source/Editor/UnrealEd/Public/ComponentVisualizer.h b/Engine/Source/Editor/UnrealEd/Public/ComponentVisualizer.h
--- a/Engine/Source/Editor/UnrealEd/Public/ComponentVisualizer.h
+++ b/Engine/Source/Editor/UnrealEd/Public/ComponentVisualizer.h
@@ -108,6 +108,8 @@ public:
 
 	/** */
 	virtual void OnRegister() {}
+	/** Only show this visualizer if the actor is selected */
+	virtual bool ShowWhenSelected() { return true; }
 	/** Draw visualization for the supplied component */
 	virtual void DrawVisualization(const UActorComponent* Component, const FSceneView* View, FPrimitiveDrawInterface* PDI) {}
 	/** Draw HUD on viewport for the supplied component */
diff --git a/Engine/Source/Editor/WorldBrowser/Private/SWorldDetails.cpp b/Engine/Source/Editor/WorldBrowser/Private/SWorldDetails.cpp
--- a/Engine/Source/Editor/WorldBrowser/Private/SWorldDetails.cpp
+++ b/Engine/Source/Editor/WorldBrowser/Private/SWorldDetails.cpp
@@ -70,6 +70,7 @@ void SWorldDetails::OnBrowseWorld(UWorld* InWorld)
 		Args.bShowActorLabel = false;
 	
 		DetailsView = PropertyModule.CreateDetailView(Args);
+		WorldDetailsView = PropertyModule.CreateDetailView(Args);
 		ChildSlot
 		[
 			SNew(SVerticalBox)
@@ -161,6 +162,18 @@ void SWorldDetails::OnBrowseWorld(UWorld* InWorld)
 					DetailsView.ToSharedRef()
 				]
 			]
+
+			// World details
+			+SVerticalBox::Slot()
+			.FillHeight(1.f)
+			.Padding(0,4,0,0)
+			[
+				SNew(SBorder)
+				.BorderImage(FEditorStyle::GetBrush(TEXT("ToolPanel.GroupBorder")))
+				[
+					WorldDetailsView.ToSharedRef()
+				]
+			]
 		];
 
 		WorldModel->RegisterDetailsCustomization(PropertyModule, DetailsView);
@@ -192,10 +205,12 @@ void SWorldDetails::OnSelectionChanged()
 	{
 		// Clear ComboBox selection in case we have multiple selection
 		SubLevelsComboBox->ClearSelection();
+		WorldDetailsView->SetObject(nullptr);
 	}
 	else
 	{
 		SubLevelsComboBox->SetSelectedItem(SelectedLevels[0]);
+		WorldDetailsView->SetObject(Cast<UObject>(SelectedLevels[0]->GetLevelObject()->GetLevelPartition()));
 	}
 
 	bUpdatingSelection = false;
diff --git a/Engine/Source/Editor/WorldBrowser/Private/SWorldDetails.h b/Engine/Source/Editor/WorldBrowser/Private/SWorldDetails.h
--- a/Engine/Source/Editor/WorldBrowser/Private/SWorldDetails.h
+++ b/Engine/Source/Editor/WorldBrowser/Private/SWorldDetails.h
@@ -58,6 +58,7 @@ private:
 private:
 	TSharedPtr<FLevelCollectionModel>				WorldModel;
 	TSharedPtr<IDetailsView>						DetailsView;
+	TSharedPtr<IDetailsView>						WorldDetailsView;
 	TSharedPtr<SComboBox<TSharedPtr<FLevelModel>>>	SubLevelsComboBox;
 	bool											bUpdatingSelection;	
 };
diff --git a/Engine/Source/Editor/WorldBrowser/Private/StreamingLevels/StreamingLevelCollectionModel.cpp b/Engine/Source/Editor/WorldBrowser/Private/StreamingLevels/StreamingLevelCollectionModel.cpp
--- a/Engine/Source/Editor/WorldBrowser/Private/StreamingLevels/StreamingLevelCollectionModel.cpp
+++ b/Engine/Source/Editor/WorldBrowser/Private/StreamingLevels/StreamingLevelCollectionModel.cpp
@@ -120,11 +120,20 @@ void FStreamingLevelCollectionModel::UnloadLevels(const FLevelModelList& InLevel
 	bool bHaveDirtyLevels = false;
 	for (const TSharedPtr<FLevelModel>& LevelModel : InLevelList)
 	{
-		if (LevelModel->IsDirty() && !LevelModel->IsLocked() && !LevelModel->IsPersistent())
+		if (!bHaveDirtyLevels && LevelModel->IsDirty() && !LevelModel->IsLocked() && !LevelModel->IsPersistent())
 		{
 			// this level is dirty and can be removed from the world
 			bHaveDirtyLevels = true;
-			break;
+		}
+
+		if (const ULevel* Level = LevelModel->GetLevelObject())
+		{
+			if (Level->IsPartitionSubLevel())
+			{
+				// this level is a partition sublevel and cannot be removed from the world
+				FMessageDialog::Open(EAppMsgType::Ok, LOCTEXT("RemoveLevel_PartitionSubLevel", "Cannot remove a level which is part of a partition. Delete the top level instead."));
+				return;
+			}
 		}
 	}
 
diff --git a/Engine/Source/Runtime/Engine/Classes/Engine/Engine.h b/Engine/Source/Runtime/Engine/Classes/Engine/Engine.h
--- a/Engine/Source/Runtime/Engine/Classes/Engine/Engine.h
+++ b/Engine/Source/Runtime/Engine/Classes/Engine/Engine.h
@@ -1853,6 +1853,13 @@ public:
 	/** Called by internal engine systems after level actors have changed to notify other subsystems */
 	void BroadcastLevelActorDeleted(AActor* InActor) { LevelActorDeletedEvent.Broadcast(InActor); }
 
+	/** Editor-only event triggered when actors outer changes */
+	DECLARE_EVENT_TwoParams(UEngine, FLevelActorOuterChangedEvent, AActor*, UObject*);
+	FLevelActorOuterChangedEvent& OnLevelActorOuterChanged() { return LevelActorOuterChangedEvent; }
+
+	/** Called by internal engine systems after level actors have changed outer */
+	void BroadcastLevelActorOuterChanged(AActor* InActor, UObject* InOldOuter) { LevelActorOuterChangedEvent.Broadcast(InActor, InOldOuter); }
+
 	/** Editor-only event triggered when actors are attached in the world */
 	DECLARE_EVENT_TwoParams( UEngine, FLevelActorAttachedEvent, AActor*, const AActor* );
 	FLevelActorAttachedEvent& OnLevelActorAttached() { return LevelActorAttachedEvent; }
@@ -1874,6 +1881,22 @@ public:
 	/** Called by internal engine systems after a level actor's folder has been changed */
 	void BroadcastLevelActorFolderChanged(const AActor* InActor, FName OldPath) { LevelActorFolderChangedEvent.Broadcast(InActor, OldPath); }
 
+	/** Editor-only event triggered when an actor is being moved, rotated or scaled (AActor::PostEditMove) */
+	DECLARE_EVENT_OneParam(UEngine, FOnActorMovingEvent, AActor*);
+	FOnActorMovingEvent& OnActorMoving() { return OnActorMovingEvent; }
+
+	/** Called by internal engine systems when an actor is being moved to notify other subsystems */
+	void BroadcastOnActorMoving(AActor* Actor) { OnActorMovingEvent.Broadcast(Actor); }
+
+	/** Editor-only event triggered after actors are moved, rotated or scaled by an editor system */
+	DECLARE_EVENT_OneParam(UEditorEngine, FOnActorsMovedEvent, TArray<AActor*>&);
+	FOnActorsMovedEvent& OnActorsMoved() { return OnActorsMovedEvent; }
+
+	/**
+	 * Called when actors have been translated, rotated, or scaled by the editor
+	 */
+	void BroadcastActorsMoved(TArray<AActor*>& Actors) const { OnActorsMovedEvent.Broadcast(Actors); }
+
 	/** Editor-only event triggered after an actor is moved, rotated or scaled (AActor::PostEditMove) */
 	DECLARE_EVENT_OneParam( UEngine, FOnActorMovedEvent, AActor* );
 	FOnActorMovedEvent& OnActorMoved() { return OnActorMovedEvent; }
@@ -2717,6 +2740,9 @@ private:
 	/** Broadcasts whenever an actor is removed. */
 	FLevelActorDeletedEvent LevelActorDeletedEvent;
 
+	/** Broadcasts whenever an actor's outer changes */
+	FLevelActorOuterChangedEvent LevelActorOuterChangedEvent;
+
 	/** Broadcasts whenever an actor is attached. */
 	FLevelActorAttachedEvent LevelActorAttachedEvent;
 
@@ -2732,8 +2758,14 @@ private:
 	/** Broadcasts whenever a component is being renamed */
 	FLevelComponentRequestRenameEvent LevelComponentRequestRenameEvent;
 
+	/** Broadcasts when an actor is being moved, rotated or scaled */
+	FOnActorMovingEvent	OnActorMovingEvent;
+
 	/** Broadcasts after an actor has been moved, rotated or scaled */
-	FOnActorMovedEvent		OnActorMovedEvent;
+	FOnActorMovedEvent	OnActorMovedEvent;
+
+	/** Broadcast when a group of actors have been moved, rotated, or scaled */
+	FOnActorsMovedEvent OnActorsMovedEvent;
 
 	/** Broadcasts after a component has been moved, rotated or scaled */
 	FOnComponentTransformChangedEvent OnComponentTransformChangedEvent;
diff --git a/Engine/Source/Runtime/Engine/Classes/Engine/Level.h b/Engine/Source/Runtime/Engine/Classes/Engine/Level.h
--- a/Engine/Source/Runtime/Engine/Classes/Engine/Level.h
+++ b/Engine/Source/Runtime/Engine/Classes/Engine/Level.h
@@ -31,6 +31,20 @@ class UTexture2D;
 struct FLevelCollection;
 class ULevelActorContainer;
 
+UINTERFACE()
+class ULevelPartitionInterface : public UInterface
+{
+	GENERATED_BODY()
+};
+
+class ILevelPartitionInterface
+{
+	GENERATED_BODY()
+
+public:
+	virtual ULevel* GetSubLevel(const FVector& Coords) const = 0;
+};
+
 /**
  * Structure containing all information needed for determining the screen space
  * size of an object/ texture instance.
@@ -595,6 +609,23 @@ public:
 
 	UPROPERTY(transient)
 	bool bLevelOkayForPlacementWhileCheckedIn;
+
+	/** Returns true if the current level is a partitioned level */
+	ENGINE_API bool IsPartitionedLevel() const;
+
+	/** Returns true if the current level is a sublevel (managed by a parent partitioned level) */
+	ENGINE_API bool IsPartitionSubLevel() const;
+
+	/** Assign a level partition to this level */
+	ENGINE_API void SetLevelPartition(ILevelPartitionInterface* LevelPartition);
+
+	/** Get the level partition assigned to this level, if any */
+	ENGINE_API ILevelPartitionInterface* GetLevelPartition();
+	ENGINE_API const ILevelPartitionInterface* GetLevelPartition() const;
+	
+	/** Setup the provided sublevel so that it is handled by this level's partition */
+	ENGINE_API void SetPartitionSubLevel(ULevel* SubLevel);
+
 #endif //WITH_EDITORONLY_DATA
 
 	/** Actor which defines level logical bounding box				*/
@@ -635,6 +666,16 @@ private:
 	UPROPERTY()
 	TArray<FReplicatedStaticActorDestructionInfo> DestroyedReplicatedStaticActors;
 
+#if WITH_EDITORONLY_DATA
+	/** Level partition, if any */
+	UPROPERTY(EditInstanceOnly, Category = World)
+	TScriptInterface<ILevelPartitionInterface> LevelPartition;
+
+	/** When the level is partitioned, this will point to the owner partition (will be the same as this->LevelPartition in case that is the top partition level */
+	UPROPERTY()
+	TSoftObjectPtr<UObject> OwnerLevelPartition;
+#endif // #if WITH_EDITORONLY_DATA
+
 public:
 	// Used internally to determine which actors should go on the world's NetworkActor list
 	ENGINE_API static bool IsNetActor(const AActor* Actor);
@@ -672,6 +713,7 @@ public:
 	virtual void PostEditUndo() override;	
 	virtual void PostEditChangeProperty(FPropertyChangedEvent& PropertyChangedEvent) override;
 	virtual void BeginCacheForCookedPlatformData(const ITargetPlatform *TargetPlatform) override;
+	virtual bool CanEditChange(const UProperty* PropertyThatWillChange) const override;
 #endif // WITH_EDITOR
 	virtual void PostLoad() override;
 	virtual void PreSave(const class ITargetPlatform* TargetPlatform) override;
@@ -825,7 +867,7 @@ public:
 	ENGINE_API void HandleLegacyMapBuildData();
 
 #if WITH_EDITOR
-	/** 
+	/**
 	*  Called after lighting was built and data gets propagated to this level
 	*  @param	bLightingSuccessful	 Whether lighting build was successful
 	*/
diff --git a/Engine/Source/Runtime/Engine/Classes/Engine/LevelBounds.h b/Engine/Source/Runtime/Engine/Classes/Engine/LevelBounds.h
--- a/Engine/Source/Runtime/Engine/Classes/Engine/LevelBounds.h
+++ b/Engine/Source/Runtime/Engine/Classes/Engine/LevelBounds.h
@@ -24,13 +24,17 @@ class FEditorTickableLevelBounds
  * Updates bounding box automatically based on actors transformation changes or holds fixed user defined bounding box
  * Uses only actors where AActor::IsLevelBoundsRelevant() == true
  */
-UCLASS(hidecategories=(Advanced, Collision, Display, Rendering, Physics, Input), showcategories=("Input|MouseInput", "Input|TouchInput"), MinimalAPI)
-class ALevelBounds 
+UCLASS(hidecategories=(Advanced, Collision, Display, Rendering, Physics, Input), showcategories=("Input|MouseInput", "Input|TouchInput"))
+class ENGINE_API ALevelBounds
 	: public AActor
 	, public FEditorTickableLevelBounds 
 {
 	GENERATED_UCLASS_BODY()
 		
+	/** Bounding box for the level bounds. */
+	UPROPERTY(EditAnywhere, Category = LevelBounds)
+	class UBoxComponent* BoxComponent;
+
 	/** Whether to automatically update actor bounds based on all relevant actors bounds belonging to the same level */
 	UPROPERTY(EditAnywhere, Category=LevelBounds)
 	bool bAutoUpdateBounds;
@@ -48,7 +52,7 @@ class ALevelBounds
 	//~ End AActor Interface.
 
 	/** @return Bounding box which includes all relevant actors bounding boxes belonging to specified level */
-	ENGINE_API static FBox CalculateLevelBounds(ULevel* InLevel);
+	static FBox CalculateLevelBounds(ULevel* InLevel);
 
 #if WITH_EDITOR
 	virtual void PostEditUndo() override;
@@ -61,12 +65,12 @@ class ALevelBounds
 	void MarkLevelBoundsDirty();
 
 	/** @return True if there were no actors contributing to bounds and we are currently using the default bounds */
-	ENGINE_API bool IsUsingDefaultBounds() const;
+	bool IsUsingDefaultBounds() const;
 
 	/** Update level bounds immediately so the bounds are accurate when returning. Use only when needed because updating the bounds is slow */
-	ENGINE_API void UpdateLevelBoundsImmediately();
+	void UpdateLevelBoundsImmediately();
 	
-private:
+protected:
 	/** FTickableGameObject interface */
 	virtual void Tick(float DeltaTime) override;
 	virtual UWorld* GetTickableGameObjectWorld() const override { return GetWorld(); }
diff --git a/Engine/Source/Runtime/Engine/Private/Actor.cpp b/Engine/Source/Runtime/Engine/Private/Actor.cpp
--- a/Engine/Source/Runtime/Engine/Private/Actor.cpp
+++ b/Engine/Source/Runtime/Engine/Private/Actor.cpp
@@ -940,6 +940,9 @@ bool AActor::Rename( const TCHAR* InName, UObject* NewOuter, ERenameFlags Flags
 		}
 	}
 
+#if WITH_EDITOR
+	UObject* OldOuter = GetOuter();
+#endif
 	const bool bSuccess = Super::Rename( InName, NewOuter, Flags );
 
 	if (!bRenameTest && bChangingOuters)
@@ -956,6 +959,13 @@ bool AActor::Rename( const TCHAR* InName, UObject* NewOuter, ERenameFlags Flags
 			}
 			RegisterAllActorTickFunctions(true, true); // register all tick functions
 		}
+
+#if WITH_EDITOR
+		if (GEngine && OldOuter != GetOuter())
+		{
+			GEngine->BroadcastLevelActorOuterChanged(this, OldOuter);
+		}
+#endif
 	}
 	return bSuccess;
 }
diff --git a/Engine/Source/Runtime/Engine/Private/ActorEditor.cpp b/Engine/Source/Runtime/Engine/Private/ActorEditor.cpp
--- a/Engine/Source/Runtime/Engine/Private/ActorEditor.cpp
+++ b/Engine/Source/Runtime/Engine/Private/ActorEditor.cpp
@@ -148,6 +148,11 @@ void AActor::PostEditMove(bool bFinished)
 		}
 	}
 
+	if (!FLevelUtils::IsMovingLevel())
+	{
+		GEngine->BroadcastOnActorMoving(this);
+	}
+
 	if ( bFinished )
 	{
 		UWorld* World = GetWorld();
diff --git a/Engine/Source/Runtime/Engine/Private/Level.cpp b/Engine/Source/Runtime/Engine/Private/Level.cpp
--- a/Engine/Source/Runtime/Engine/Private/Level.cpp
+++ b/Engine/Source/Runtime/Engine/Private/Level.cpp
@@ -1981,6 +1981,32 @@ void ULevel::BeginCacheForCookedPlatformData(const ITargetPlatform *TargetPlatfo
 	}
 }
 
+bool ULevel::CanEditChange(const UProperty* PropertyThatWillChange) const
+{
+	static FName NAME_LevelPartition = GET_MEMBER_NAME_CHECKED(ULevel, LevelPartition);
+	if (PropertyThatWillChange->GetFName() == NAME_LevelPartition)
+	{
+		// Can't set a partition on the persistent level
+		if (IsPersistentLevel())
+		{
+			return false;
+		}
+
+		// Can't set a partition on partition sublevels
+		if (IsPartitionSubLevel())
+		{
+			return false;
+		}
+
+		// Can't set a partition if using world composition
+		if (WorldSettings && WorldSettings->bEnableWorldComposition)
+		{
+			return false;
+		}
+	}
+	return Super::CanEditChange(PropertyThatWillChange);
+}
+
 void ULevel::FixupForPIE(int32 PIEInstanceID)
 {
 	FTemporaryPlayInEditorIDOverride SetPlayInEditorID(PIEInstanceID);
@@ -2222,4 +2248,39 @@ bool ULevel::HasVisibilityChangeRequestPending() const
 	return (OwningWorld && ( this == OwningWorld->GetCurrentLevelPendingVisibility() || this == OwningWorld->GetCurrentLevelPendingInvisibility() ) );
 }
 
+#if WITH_EDITORONLY_DATA
+
+bool ULevel::IsPartitionedLevel() const
+{
+	return LevelPartition != nullptr;
+}
+
+bool ULevel::IsPartitionSubLevel() const
+{
+	return OwnerLevelPartition.IsValid() && LevelPartition == nullptr;
+}
+
+void ULevel::SetLevelPartition(ILevelPartitionInterface* InLevelPartition)
+{
+	UObject* PartitionObject = Cast<UObject>(InLevelPartition);
+	LevelPartition = PartitionObject;
+	OwnerLevelPartition = PartitionObject;
+}
+
+ILevelPartitionInterface* ULevel::GetLevelPartition()
+{
+	return Cast<ILevelPartitionInterface>(OwnerLevelPartition.Get());
+}
+
+const ILevelPartitionInterface* ULevel::GetLevelPartition() const
+{
+	return Cast<ILevelPartitionInterface>(OwnerLevelPartition.Get());
+}
+
+void ULevel::SetPartitionSubLevel(ULevel* SubLevel)
+{
+	check(LevelPartition);
+	SubLevel->OwnerLevelPartition = Cast<UObject>(&*LevelPartition);
+}
 
+#endif // #if WITH_EDITORONLY_DATA
\ No newline at end of file
diff --git a/Engine/Source/Runtime/Engine/Private/LevelBounds.cpp b/Engine/Source/Runtime/Engine/Private/LevelBounds.cpp
--- a/Engine/Source/Runtime/Engine/Private/LevelBounds.cpp
+++ b/Engine/Source/Runtime/Engine/Private/LevelBounds.cpp
@@ -12,11 +12,11 @@ static const FVector DefaultLevelSize = FVector(1000.f);
 ALevelBounds::ALevelBounds(const FObjectInitializer& ObjectInitializer)
 	: Super(ObjectInitializer)
 {
-	UBoxComponent* BoxComponent = CreateDefaultSubobject<UBoxComponent>(TEXT("BoxComponent0"));
+	BoxComponent = CreateDefaultSubobject<UBoxComponent>(TEXT("BoxComponent0"));
 	RootComponent = BoxComponent;
-	RootComponent->Mobility = EComponentMobility::Movable;
-	RootComponent->SetRelativeScale3D(DefaultLevelSize);
-	RootComponent->SetCanEverAffectNavigation(false);
+	BoxComponent->Mobility = EComponentMobility::Movable;
+	BoxComponent->SetRelativeScale3D(DefaultLevelSize);
+	BoxComponent->SetCanEverAffectNavigation(false);
 
 	bAutoUpdateBounds = true;
 
diff --git a/Engine/Source/Runtime/Foliage/Private/FoliageActor.cpp b/Engine/Source/Runtime/Foliage/Private/FoliageActor.cpp
--- a/Engine/Source/Runtime/Foliage/Private/FoliageActor.cpp
+++ b/Engine/Source/Runtime/Foliage/Private/FoliageActor.cpp
@@ -4,6 +4,7 @@
 #include "InstancedFoliageActor.h"
 #include "FoliageType_Actor.h"
 #include "FoliageHelper.h"
+#include "Engine/Engine.h"
 
 //
 //
@@ -85,9 +86,10 @@ AActor* FFoliageActor::Spawn(AInstancedFoliageActor* IFA, const FFoliageInstance
 	SpawnParameters.ObjectFlags = RF_Transactional;
 	SpawnParameters.bHideFromSceneOutliner = true;
 	SpawnParameters.OverrideLevel = IFA->GetLevel();
-	AActor* NewActor = IFA->GetWorld()->SpawnActor<AActor>(ActorClass, Instance.GetInstanceWorldTransform(), SpawnParameters);
+	AActor* NewActor = IFA->GetWorld()->SpawnActor(ActorClass, nullptr, nullptr, SpawnParameters);
 	if (NewActor)
 	{
+		NewActor->SetActorTransform(Instance.GetInstanceWorldTransform());
 		FFoliageHelper::SetIsOwnedByFoliage(NewActor);
 	}
 	return NewActor;
@@ -128,15 +130,29 @@ void FFoliageActor::AddInstance(AInstancedFoliageActor* IFA, const FFoliageInsta
 	ActorInstances.Add(Spawn(IFA, NewInstance));
 }
 
+void FFoliageActor::AddExistingInstance(AInstancedFoliageActor* IFA, const FFoliageInstance& ExistingInstance, UObject* InstanceImplementation)
+{
+	AActor* Actor = Cast<AActor>(InstanceImplementation);
+	check(Actor);
+	check(Actor->GetClass() == ActorClass);
+	Actor->SetActorTransform(ExistingInstance.GetInstanceWorldTransform());
+	FFoliageHelper::SetIsOwnedByFoliage(Actor);
+	check(IFA->GetLevel() == Actor->GetLevel());
+	ActorInstances.Add(Actor);
+}
+
 void FFoliageActor::RemoveInstance(int32 InstanceIndex)
 {
 	AActor* Actor = ActorInstances[InstanceIndex];
 	ActorInstances.RemoveAtSwap(InstanceIndex);
+	Actor->GetWorld()->DestroyActor(Actor, true);
+}
 
-	if (Actor && Actor->GetWorld())
-	{
-		Actor->GetWorld()->DestroyActor(Actor, true);
-	}
+void FFoliageActor::MoveInstance(int32 InstanceIndex, UObject*& OutInstanceImplementation)
+{
+	AActor* Actor = ActorInstances[InstanceIndex];
+	OutInstanceImplementation = Actor;
+	ActorInstances.RemoveAtSwap(InstanceIndex);
 }
 
 void FFoliageActor::SetInstanceWorldTransform(int32 InstanceIndex, const FTransform& Transform, bool bTeleport)
@@ -213,6 +229,37 @@ void FFoliageActor::PostEditUndo(AInstancedFoliageActor* IFA, UFoliageType* Foli
 	UpdateActorTransforms(Instances);
 }
 
+void FFoliageActor::PreMoveInstances(AInstancedFoliageActor* InIFA, const TArray<int32>& InInstancesMoved)
+{
+	for (int32 Index : InInstancesMoved)
+	{
+		if (AActor* Actor = ActorInstances[Index])
+		{
+			Actor->Modify();
+		}
+	}
+}
+
+void FFoliageActor::PostMoveInstances(AInstancedFoliageActor* InIFA, const TArray<int32>& InInstancesMoved, bool bFinished)
+{
+	// Copy because moving actors might remove them from ActorInstances
+	TArray<AActor*> MovedActors;
+	MovedActors.Reserve(InInstancesMoved.Num());
+	for (int32 Index : InInstancesMoved)
+	{
+		if (AActor * Actor = ActorInstances[Index])
+		{
+			MovedActors.Add(Actor);
+			Actor->PostEditMove(bFinished);
+		}
+	}
+
+	if (GIsEditor && GEngine && MovedActors.Num() && bFinished)
+	{
+		GEngine->BroadcastActorsMoved(MovedActors);
+	}
+}
+
 void FFoliageActor::NotifyFoliageTypeChanged(AInstancedFoliageActor* IFA, UFoliageType* FoliageType, const TArray<FFoliageInstance>& Instances, const TSet<int32>& SelectedIndices, bool bSourceChanged)
 {
 	if (UFoliageType_Actor* FoliageTypeActor = Cast<UFoliageType_Actor>(FoliageType))
diff --git a/Engine/Source/Runtime/Foliage/Private/FoliageActor.h b/Engine/Source/Runtime/Foliage/Private/FoliageActor.h
--- a/Engine/Source/Runtime/Foliage/Private/FoliageActor.h
+++ b/Engine/Source/Runtime/Foliage/Private/FoliageActor.h
@@ -30,7 +30,9 @@ struct FFoliageActor : public FFoliageImpl
 	virtual int32 GetInstanceCount() const override;
 	virtual void PreAddInstances(AInstancedFoliageActor* IFA, const UFoliageType* FoliageType, int32 Count) override;
 	virtual void AddInstance(AInstancedFoliageActor* IFA, const FFoliageInstance& NewInstance) override;
+	virtual void AddExistingInstance(AInstancedFoliageActor* IFA, const FFoliageInstance& ExistingInstance, UObject* InstanceImplementation) override;
 	virtual void RemoveInstance(int32 InstanceIndex) override;
+	virtual void MoveInstance(int32 InstanceIndex, UObject*& OutInstanceImplementation) override;
 	virtual void SetInstanceWorldTransform(int32 InstanceIndex, const FTransform& Transform, bool bTeleport) override;
 	virtual FTransform GetInstanceWorldTransform(int32 InstanceIndex) const override;
 	virtual bool IsOwnedComponent(const UPrimitiveComponent* Component) const override;
@@ -45,6 +47,8 @@ struct FFoliageActor : public FFoliageImpl
 	virtual void Refresh(AInstancedFoliageActor* IFA, const TArray<FFoliageInstance>& Instances, bool Async, bool Force) override;
 	virtual void OnHiddenEditorViewMaskChanged(uint64 InHiddenEditorViews) override;
 	virtual void PostEditUndo(AInstancedFoliageActor* IFA, UFoliageType* FoliageType, const TArray<FFoliageInstance>& Instances, const TSet<int32>& SelectedIndices) override;
+	virtual void PreMoveInstances(AInstancedFoliageActor* InIFA, const TArray<int32>& InInstancesMoved) override;
+	virtual void PostMoveInstances(AInstancedFoliageActor* InIFA, const TArray<int32>& InInstancesMoved, bool bFinished) override;
 	virtual void NotifyFoliageTypeChanged(AInstancedFoliageActor* IFA, UFoliageType* FoliageType, const TArray<FFoliageInstance>& Instances, const TSet<int32>& SelectedIndices, bool bSourceChanged) override;
 	void Reapply(AInstancedFoliageActor* IFA, const UFoliageType* FoliageType, const TArray<FFoliageInstance>& Instances, bool bPostLoad = false);
 	AActor* Spawn(AInstancedFoliageActor* IFA, const FFoliageInstance& Instance);
diff --git a/Engine/Source/Runtime/Foliage/Private/InstancedFoliage.cpp b/Engine/Source/Runtime/Foliage/Private/InstancedFoliage.cpp
--- a/Engine/Source/Runtime/Foliage/Private/InstancedFoliage.cpp
+++ b/Engine/Source/Runtime/Foliage/Private/InstancedFoliage.cpp
@@ -47,6 +47,7 @@ InstancedFoliage.cpp: Instanced foliage implementation.
 #include "FoliageActor.h"
 #include "LevelUtils.h"
 #include "FoliageHelper.h"
+#include "Algo/Transform.h"
 
 #define LOCTEXT_NAMESPACE "InstancedFoliage"
 
@@ -1725,7 +1726,7 @@ void FFoliageInfo::SetInstanceWorldTransform(int32 InstanceIndex, const FTransfo
 	Implementation->SetInstanceWorldTransform(InstanceIndex, Transform, bTeleport);
 }
 
-void FFoliageInfo::AddInstanceImpl(AInstancedFoliageActor* InIFA, const FFoliageInstance& InNewInstance)
+void FFoliageInfo::AddInstanceImpl(AInstancedFoliageActor* InIFA, const FFoliageInstance& InNewInstance, FFoliageInfo::FAddImplementationFunc ImplementationFunc)
 {
 	// Add the instance taking either a free slot or adding a new item.
 	int32 InstanceIndex = Instances.Add(InNewInstance);
@@ -1742,10 +1743,15 @@ void FFoliageInfo::AddInstanceImpl(AInstancedFoliageActor* InIFA, const FFoliage
 	InstanceHash->InsertInstance(AddedInstance.Location, InstanceIndex);
 
 	// Add the instance to the component
-	Implementation->AddInstance(InIFA, AddedInstance);
+	ImplementationFunc(Implementation.Get(), InIFA, AddedInstance);
 }
 
-void FFoliageInfo::AddInstances(AInstancedFoliageActor* InIFA, const UFoliageType* InSettings, const TSet<const FFoliageInstance*>& InNewInstances)
+void FFoliageInfo::AddInstances(AInstancedFoliageActor* InIFA, const UFoliageType* InSettings, const TArray<const FFoliageInstance*>& InNewInstances)
+{
+	AddInstancesImpl(InIFA, InSettings, InNewInstances, [](FFoliageImpl* Impl, AInstancedFoliageActor* LocalIFA, const FFoliageInstance& LocalInstance) { Impl->AddInstance(LocalIFA, LocalInstance); });
+}
+
+void FFoliageInfo::AddInstancesImpl(AInstancedFoliageActor* InIFA, const UFoliageType* InSettings, const TArray<const FFoliageInstance*>& InNewInstances, FFoliageInfo::FAddImplementationFunc ImplementationFunc)
 {
 	SCOPE_CYCLE_COUNTER(STAT_FoliageAddInstance);
 
@@ -1756,10 +1762,10 @@ void FFoliageInfo::AddInstances(AInstancedFoliageActor* InIFA, const UFoliageTyp
 	Implementation->BeginUpdate();
 
 	Instances.Reserve(Instances.Num() + InNewInstances.Num());
-			
+
 	for (const FFoliageInstance* Instance : InNewInstances)
 	{
-		AddInstanceImpl(InIFA, *Instance);
+		AddInstanceImpl(InIFA, *Instance, ImplementationFunc);
 	}
 
 	CheckValid();
@@ -1780,22 +1786,80 @@ void FFoliageInfo::AddInstance(AInstancedFoliageActor* InIFA, const UFoliageType
 
 void FFoliageInfo::AddInstance(AInstancedFoliageActor* InIFA, const UFoliageType* InSettings, const FFoliageInstance& InNewInstance)
 {
-	SCOPE_CYCLE_COUNTER(STAT_FoliageAddInstance);
+	AddInstances(InIFA, InSettings, { &InNewInstance });
+}
 
-	InIFA->Modify();
+void FFoliageInfo::MoveInstances(AInstancedFoliageActor* InFromIFA, AInstancedFoliageActor* InToIFA, const TSet<int32>& InInstancesToMove, bool bKeepSelection)
+{
+	const UFoliageType* FoliageType = InFromIFA->FindFoliageType(this);
+	check(FoliageType);
 
-	Implementation->PreAddInstances(InIFA, InSettings, 1);
+	InFromIFA->Modify();
+	InToIFA->Modify();
+	
+	FFoliageInfo* OutFoliageInfo = nullptr;
+	UFoliageType* ToFoliageType = InToIFA->AddFoliageType(FoliageType, &OutFoliageInfo);
 
-	Implementation->BeginUpdate();
+	TArray<int32> NewSelectedIndices;
+	
+	struct FFoliageMoveInstance : FFoliageInstance
+	{
+		FFoliageMoveInstance(const FFoliageInstance& InInstance, UActorComponent* InBaseComponent)
+			: FFoliageInstance(InInstance)
+		{
+			BaseComponent = InBaseComponent;
+		}
+		UObject* InstanceImplementation = nullptr;
+	};
 
-	AddInstanceImpl(InIFA, InNewInstance);
+	TArray<int32> InstancesToMove = InInstancesToMove.Array();
+	TMap<int32, FFoliageMoveInstance> MoveData;
+	MoveData.Reserve(InInstancesToMove.Num());
+	
+	for (int32 InstanceIndex : InInstancesToMove)
+	{
+		MoveData.Add(InstanceIndex, FFoliageMoveInstance(Instances[InstanceIndex], InFromIFA->GetBaseComponentFromBaseId(Instances[InstanceIndex].BaseId)));
+		if (bKeepSelection && SelectedIndices.Contains(InstanceIndex))
+		{
+			NewSelectedIndices.Add(OutFoliageInfo->Instances.Num() + MoveData.Num() - 1);
+		}
+	}
 
-	CheckValid();
+	RemoveInstancesImpl(InFromIFA, InstancesToMove, true, [&MoveData](FFoliageImpl* Impl, int32 Index)
+	{
+		Impl->MoveInstance(Index, MoveData[Index].InstanceImplementation);
+	});
 
-	Implementation->EndUpdate();
+	TArray<const FFoliageInstance*> NewInstances;
+	NewInstances.Reserve(MoveData.Num());
+	Algo::Transform(MoveData, NewInstances, [](const TPair<int32, FFoliageMoveInstance>& Pair)
+	{
+		return &Pair.Value;
+	});
+	
+	int32 AddedIndex = 0;
+	OutFoliageInfo->AddInstancesImpl(InToIFA, ToFoliageType, NewInstances, [&](FFoliageImpl* Impl, AInstancedFoliageActor* LocalIFA, const FFoliageInstance& LocalInstance)
+	{
+		const FFoliageMoveInstance* MoveInstance = StaticCast<const FFoliageMoveInstance*>(NewInstances[AddedIndex]);
+		Impl->AddExistingInstance(LocalIFA, LocalInstance, MoveInstance->InstanceImplementation);
+		AddedIndex++;
+	});
+			
+	OutFoliageInfo->Refresh(InToIFA, true, true);
+
+	// Select if needed
+	if (NewSelectedIndices.Num())
+	{
+		OutFoliageInfo->SelectInstances(InToIFA, true, NewSelectedIndices);
+	}
 }
 
 void FFoliageInfo::RemoveInstances(AInstancedFoliageActor* InIFA, const TArray<int32>& InInstancesToRemove, bool RebuildFoliageTree)
+{
+	RemoveInstancesImpl(InIFA, InInstancesToRemove, RebuildFoliageTree, [](FFoliageImpl* Impl, int32 Index) { Impl->RemoveInstance(Index); });
+}
+
+void FFoliageInfo::RemoveInstancesImpl(AInstancedFoliageActor* InIFA, const TArray<int32>& InInstancesToRemove, bool RebuildFoliageTree, FFoliageInfo::FRemoveImplementationFunc ImplementationFunc)
 {
 	SCOPE_CYCLE_COUNTER(STAT_FoliageRemoveInstance);
 
@@ -1826,7 +1890,7 @@ void FFoliageInfo::RemoveInstances(AInstancedFoliageActor* InIFA, const TArray<i
 		InstanceHash->RemoveInstance(Instance.Location, InstanceIndex);
 
 		// remove from the component
-		Implementation->RemoveInstance(InstanceIndex);
+		ImplementationFunc(Implementation.Get(), InstanceIndex);
 
 		// Remove it from the selection.
 		SelectedIndices.Remove(InstanceIndex);
@@ -1893,6 +1957,8 @@ void FFoliageInfo::PreMoveInstances(AInstancedFoliageActor* InIFA, const TArray<
 		const FFoliageInstance& Instance = Instances[InstanceIndex];
 		InstanceHash->RemoveInstance(Instance.Location, InstanceIndex);
 	}
+
+	Implementation->PreMoveInstances(InIFA, InInstancesToMove);
 }
 
 
@@ -1933,10 +1999,10 @@ void FFoliageInfo::PostUpdateInstances(AInstancedFoliageActor* InIFA, const TArr
 	}
 }
 
-void FFoliageInfo::PostMoveInstances(AInstancedFoliageActor* InIFA, const TArray<int32>& InInstancesMoved)
+void FFoliageInfo::PostMoveInstances(AInstancedFoliageActor* InIFA, const TArray<int32>& InInstancesMoved, bool bFinished)
 {
 	PostUpdateInstances(InIFA, InInstancesMoved, true, true);
-
+	Implementation->PostMoveInstances(InIFA, InInstancesMoved, bFinished);
 	bMovingInstances = false;
 }
 
@@ -2670,7 +2736,7 @@ void AInstancedFoliageActor::MoveInstancesToNewComponent(UPrimitiveComponent* In
 		return;
 	}
 
-	AInstancedFoliageActor* TargetIFA = AInstancedFoliageActor::GetInstancedFoliageActorForLevel(InNewComponent->GetTypedOuter<ULevel>(), true);
+	AInstancedFoliageActor* TargetIFA = AInstancedFoliageActor::GetInstancedFoliageActorForLevel(InNewComponent->GetOwner()->GetLevel(), true);
 	TArray<int32> InstancesToMove;
 
 	for (auto& Pair : FoliageInfos)
@@ -2703,7 +2769,7 @@ void AInstancedFoliageActor::MoveInstancesToNewComponent(UPrimitiveComponent* In
 
 void AInstancedFoliageActor::MoveInstancesToNewComponent(UPrimitiveComponent* InOldComponent, UPrimitiveComponent* InNewComponent)
 {
-	AInstancedFoliageActor* TargetIFA = AInstancedFoliageActor::GetInstancedFoliageActorForLevel(InNewComponent->GetTypedOuter<ULevel>(), true);
+	AInstancedFoliageActor* TargetIFA = AInstancedFoliageActor::GetInstancedFoliageActorForLevel(InNewComponent->GetOwner()->GetLevel(), true);
 
 	const auto OldBaseId = this->InstanceBaseCache.GetInstanceBaseId(InOldComponent);
 	if (OldBaseId == FFoliageInstanceBaseCache::InvalidBaseId)
@@ -2773,7 +2839,7 @@ void AInstancedFoliageActor::MoveInstancesToNewComponent(UWorld* InWorld, UPrimi
 	}
 }
 
-void AInstancedFoliageActor::MoveInstancesToLevel(ULevel* InTargetLevel, TSet<int32>& InInstanceList, FFoliageInfo* InCurrentMeshInfo, UFoliageType* InFoliageType)
+void AInstancedFoliageActor::MoveInstancesToLevel(ULevel* InTargetLevel, TSet<int32>& InInstanceList, FFoliageInfo* InCurrentMeshInfo, UFoliageType* InFoliageType, bool bSelect)
 {
 	if (InTargetLevel == GetLevel())
 	{
@@ -2781,25 +2847,7 @@ void AInstancedFoliageActor::MoveInstancesToLevel(ULevel* InTargetLevel, TSet<in
 	}
 
 	AInstancedFoliageActor* TargetIFA = GetInstancedFoliageActorForLevel(InTargetLevel, /*bCreateIfNone*/ true);
-
-	Modify();
-	TargetIFA->Modify();
-
-	// Do move
-	FFoliageInfo* TargetMeshInfo = nullptr;
-	UFoliageType* TargetFoliageType = TargetIFA->AddFoliageType(InFoliageType, &TargetMeshInfo);
-
-	// Add selected instances to the target actor
-	for (int32 InstanceIndex : InInstanceList)
-	{
-		FFoliageInstance& Instance = InCurrentMeshInfo->Instances[InstanceIndex];
-		TargetMeshInfo->AddInstance(TargetIFA, TargetFoliageType, Instance, InstanceBaseCache.GetInstanceBasePtr(Instance.BaseId).Get());
-	}
-
-	TargetMeshInfo->Refresh(TargetIFA, true, true);
-
-	// Remove selected instances from this actor
-	InCurrentMeshInfo->RemoveInstances(this, InInstanceList.Array(), true);
+	InCurrentMeshInfo->MoveInstances(this, TargetIFA, InInstanceList, bSelect);
 }
 
 void AInstancedFoliageActor::MoveSelectedInstancesToLevel(ULevel* InTargetLevel)
@@ -3152,6 +3200,19 @@ bool AInstancedFoliageActor::HasSelectedInstances() const
 	return false;
 }
 
+const UFoliageType* AInstancedFoliageActor::FindFoliageType(const FFoliageInfo* InFoliageInfo) const
+{
+	for (auto& Pair : FoliageInfos)
+	{
+		if (&Pair.Value.Get() == InFoliageInfo)
+		{
+			return Pair.Key;
+		}
+	}
+
+	return nullptr;
+}
+
 TMap<UFoliageType*, FFoliageInfo*> AInstancedFoliageActor::GetAllInstancesFoliageType()
 {
 	TMap<UFoliageType*, FFoliageInfo*> InstanceFoliageTypes;
@@ -3261,19 +3322,21 @@ void AInstancedFoliageActor::ApplySelection(bool bApply)
 	}
 }
 
-bool AInstancedFoliageActor::GetSelectionLocation(FVector& OutLocation) const
+bool AInstancedFoliageActor::GetSelectionLocation(FBox& OutBox) const
 {
+	// Could probably be cached instead of recalculated always.
+	bool bHasSelection = false;
 	for (const auto& Pair : FoliageInfos)
 	{
 		const FFoliageInfo& Info = Pair.Value.Get();
-		if (Info.SelectedIndices.Num())
+		for(int32 InstanceIdx : Info.SelectedIndices)
 		{
-			const int32 InstanceIdx = (*Info.SelectedIndices.CreateConstIterator());
-			OutLocation = Info.Instances[InstanceIdx].Location;
-			return true;
+			OutBox += Info.Instances[InstanceIdx].Location;
+			bHasSelection = true;
 		}
 	}
-	return false;
+		
+	return bHasSelection;
 }
 
 bool AInstancedFoliageActor::HasFoliageAttached(UActorComponent* InComponent)
@@ -3469,6 +3532,11 @@ void AInstancedFoliageActor::Serialize(FArchive& Ar)
 }
 
 #if WITH_EDITOR
+UActorComponent* AInstancedFoliageActor::GetBaseComponentFromBaseId(const FFoliageInstanceBaseId& BaseId) const
+{
+	return InstanceBaseCache.GetInstanceBasePtr(BaseId).Get();
+}
+
 AInstancedFoliageActor::FOnSelectionChanged AInstancedFoliageActor::SelectionChanged;
 AInstancedFoliageActor::FOnInstanceCoundChanged AInstancedFoliageActor::InstanceCountChanged;
 
@@ -3489,6 +3557,9 @@ void AInstancedFoliageActor::PostInitProperties()
 			OnApplyLevelTransformDelegateHandle = GetLevel()->OnApplyLevelTransform.AddUObject(this, &AInstancedFoliageActor::OnApplyLevelTransform);
 		}
 
+		GEngine->OnLevelActorOuterChanged().Remove(OnLevelActorOuterChangedDelegateHandle);
+		OnLevelActorOuterChangedDelegateHandle = GEngine->OnLevelActorOuterChanged().AddUObject(this, &AInstancedFoliageActor::OnLevelActorOuterChanged);
+
 		FWorldDelegates::PostApplyLevelOffset.Remove(OnPostApplyLevelOffsetDelegateHandle);
 		OnPostApplyLevelOffsetDelegateHandle = FWorldDelegates::PostApplyLevelOffset.AddUObject(this, &AInstancedFoliageActor::OnPostApplyLevelOffset);
 
@@ -3505,7 +3576,8 @@ void AInstancedFoliageActor::BeginDestroy()
 	{
 		GEngine->OnActorMoved().Remove(OnLevelActorMovedDelegateHandle);
 		GEngine->OnLevelActorDeleted().Remove(OnLevelActorDeletedDelegateHandle);
-		
+		GEngine->OnLevelActorOuterChanged().Remove(OnLevelActorOuterChangedDelegateHandle);
+
 		if (GetLevel())
 		{
 			GetLevel()->OnApplyLevelTransform.Remove(OnApplyLevelTransformDelegateHandle);
@@ -3929,6 +4001,51 @@ void AInstancedFoliageActor::OnLevelActorMoved(AActor* InActor)
 	}
 }
 
+void AInstancedFoliageActor::OnLevelActorOuterChanged(AActor* InActor, UObject* OldOuter)
+{
+	ULevel* OldLevel = Cast<ULevel>(OldOuter);
+
+	if (InActor->GetLevel() == OldLevel)
+	{
+		return;
+	}
+
+	if (!FFoliageHelper::IsOwnedByFoliage(InActor))
+	{
+		return;
+	}
+
+	if (OldLevel)
+	{
+		AInstancedFoliageActor* OldIFA = AInstancedFoliageActor::GetInstancedFoliageActorForLevel(OldLevel, false);
+		check(OldIFA);
+		TSet<int32> InstanceToMove;
+		FFoliageInfo* OldFoliageInfo = nullptr;
+		UFoliageType* FoliageType = nullptr;
+
+		for (auto& Pair : OldIFA->FoliageInfos)
+		{
+			if (Pair.Value->Type == EFoliageImplType::Actor)
+			{
+				FFoliageActor* FoliageActor = StaticCast<FFoliageActor*>(Pair.Value->Implementation.Get());
+				int32 Index = FoliageActor->FindIndex(InActor);
+				if (Index != INDEX_NONE)
+				{
+					InstanceToMove.Add(Index);
+					OldFoliageInfo = &Pair.Value.Get();
+					FoliageType = Pair.Key;
+					break;
+				}
+			}
+		}
+
+		if (InstanceToMove.Num() > 0)
+		{
+			OldIFA->MoveInstancesToLevel(InActor->GetLevel(), InstanceToMove, OldFoliageInfo, FoliageType, true);
+		}
+	}
+}
+
 void AInstancedFoliageActor::OnLevelActorDeleted(AActor* InActor)
 {
 	UWorld* InWorld = InActor->GetWorld();
diff --git a/Engine/Source/Runtime/Foliage/Public/InstancedFoliage.h b/Engine/Source/Runtime/Foliage/Public/InstancedFoliage.h
--- a/Engine/Source/Runtime/Foliage/Public/InstancedFoliage.h
+++ b/Engine/Source/Runtime/Foliage/Public/InstancedFoliage.h
@@ -186,9 +186,13 @@ struct FFoliageImpl
 	virtual void PreAddInstances(AInstancedFoliageActor* IFA, const UFoliageType* FoliageType, int32 Count) = 0;
 	virtual void AddInstance(AInstancedFoliageActor* IFA, const FFoliageInstance& NewInstance) = 0;
 	virtual void RemoveInstance(int32 InstanceIndex) = 0;
+	virtual void MoveInstance(int32 InstanceIndex, UObject*& OutInstanceImplementation) { RemoveInstance(InstanceIndex); }
+	virtual void AddExistingInstance(AInstancedFoliageActor* IFA, const FFoliageInstance& ExistingInstance, UObject* InstanceImplementation) { AddInstance(IFA, ExistingInstance); }
 	virtual void SetInstanceWorldTransform(int32 InstanceIndex, const FTransform& Transform, bool bTeleport) = 0;
 	virtual FTransform GetInstanceWorldTransform(int32 InstanceIndex) const = 0;
 	virtual void PostUpdateInstances() {}
+	virtual void PreMoveInstances(AInstancedFoliageActor* InIFA, const TArray<int32>& InInstancesMoved) {}
+	virtual void PostMoveInstances(AInstancedFoliageActor* InIFA, const TArray<int32>& InInstancesMoved, bool bFinished) {}
 	virtual bool IsOwnedComponent(const UPrimitiveComponent* Component) const = 0;
 	virtual int32 FindIndex(const UPrimitiveComponent* HitComponent) const { return INDEX_NONE; }
 
@@ -279,12 +283,15 @@ struct FFoliageInfo
 	FOLIAGE_API void SetRandomSeed(int32 seed);
 	FOLIAGE_API void AddInstance(AInstancedFoliageActor* InIFA, const UFoliageType* InSettings, const FFoliageInstance& InNewInstance);
 	FOLIAGE_API void AddInstance(AInstancedFoliageActor* InIFA, const UFoliageType* InSettings, const FFoliageInstance& InNewInstance, UActorComponent* InBaseComponent);
-	FOLIAGE_API void AddInstances(AInstancedFoliageActor* InIFA, const UFoliageType* InSettings, const TSet<const FFoliageInstance*>& InNewInstances);
+	FOLIAGE_API void AddInstances(AInstancedFoliageActor* InIFA, const UFoliageType* InSettings, const TArray<const FFoliageInstance*>& InNewInstances);
 
 	FOLIAGE_API void RemoveInstances(AInstancedFoliageActor* InIFA, const TArray<int32>& InInstancesToRemove, bool RebuildFoliageTree);
+
+	FOLIAGE_API void MoveInstances(AInstancedFoliageActor* InFromIFA, AInstancedFoliageActor* InToIFA, const TSet<int32>& InInstancesToMove, bool bKeepSelection);
+
 	// Apply changes in the FoliageType to the component
 	FOLIAGE_API void PreMoveInstances(AInstancedFoliageActor* InIFA, const TArray<int32>& InInstancesToMove);
-	FOLIAGE_API void PostMoveInstances(AInstancedFoliageActor* InIFA, const TArray<int32>& InInstancesMoved);
+	FOLIAGE_API void PostMoveInstances(AInstancedFoliageActor* InIFA, const TArray<int32>& InInstancesMoved, bool bFinished = false);
 	FOLIAGE_API void PostUpdateInstances(AInstancedFoliageActor* InIFA, const TArray<int32>& InInstancesUpdated, bool bReAddToHash = false, bool InUpdateSelection = false);
 	FOLIAGE_API void DuplicateInstances(AInstancedFoliageActor* InIFA, UFoliageType* InSettings, const TArray<int32>& InInstancesToDuplicate);
 	FOLIAGE_API void GetInstancesInsideSphere(const FSphere& Sphere, TArray<int32>& OutInstances);
@@ -329,7 +336,12 @@ struct FFoliageInfo
 	FFoliageInfo& operator=(const FFoliageInfo&) = delete;
 
 private:
-	void AddInstanceImpl(AInstancedFoliageActor* InIFA, const FFoliageInstance& InNewInstance);
+	using FAddImplementationFunc = TFunctionRef<void(FFoliageImpl*, AInstancedFoliageActor*, const FFoliageInstance&)>;
+	void AddInstancesImpl(AInstancedFoliageActor* InIFA, const UFoliageType* InSettings, const TArray<const FFoliageInstance*>& InNewInstances, FFoliageInfo::FAddImplementationFunc ImplementationFunc);
+	void AddInstanceImpl(AInstancedFoliageActor* InIFA, const FFoliageInstance& InNewInstance, FAddImplementationFunc ImplementationFunc);
+
+	using FRemoveImplementationFunc = TFunctionRef<void(FFoliageImpl*, int32)>;
+	void RemoveInstancesImpl(AInstancedFoliageActor* InIFA, const TArray<int32>& InInstancesToRemove, bool RebuildFoliageTree, FRemoveImplementationFunc ImplementationFunc);
 };
 
 
diff --git a/Engine/Source/Runtime/Foliage/Public/InstancedFoliageActor.h b/Engine/Source/Runtime/Foliage/Public/InstancedFoliageActor.h
--- a/Engine/Source/Runtime/Foliage/Public/InstancedFoliageActor.h
+++ b/Engine/Source/Runtime/Foliage/Public/InstancedFoliageActor.h
@@ -74,6 +74,8 @@ public:
 #if WITH_EDITORONLY_DATA
 	// Cross level references cache for instances base
 	FFoliageInstanceBaseCache InstanceBaseCache;
+
+	UActorComponent* GetBaseComponentFromBaseId(const FFoliageInstanceBaseId& BaseId) const;
 #endif// WITH_EDITORONLY_DATA
 
 	TMap<UFoliageType*, TUniqueObj<FFoliageInfo>> FoliageInfos;
@@ -181,6 +183,9 @@ public:
 	// Move all instances to a foliage actor in target level
 	FOLIAGE_API void MoveAllInstancesToLevel(ULevel* InTargetLevel);
 	
+	// Move instances to a foliage actor in target level
+	FOLIAGE_API void MoveInstancesToLevel(ULevel* InTargetLevel, TSet<int32>& InInstanceList, FFoliageInfo* InCurrentMeshInfo, UFoliageType* InFoliageType, bool bSelect = false);
+
 	// Move instances based on a component that has just been moved.
 	void MoveInstancesForMovedComponent(UActorComponent* InComponent);
 
@@ -226,6 +231,9 @@ public:
 	// Will return all the foliage type used by currently selected instances
 	FOLIAGE_API TMap<UFoliageType*, FFoliageInfo*> GetSelectedInstancesFoliageType();
 
+	// Returns FoliageType associated to this FoliageInfo
+	FOLIAGE_API const UFoliageType* FindFoliageType(const FFoliageInfo* InFoliageInfo) const;
+
 	// Will return all the foliage type used
 	FOLIAGE_API TMap<UFoliageType*, FFoliageInfo*> GetAllInstancesFoliageType();
 
@@ -233,7 +241,7 @@ public:
 	FOLIAGE_API void ApplySelection(bool bApply);
 
 	// Returns the location for the widget
-	FOLIAGE_API bool GetSelectionLocation(FVector& OutLocation) const;
+	FOLIAGE_API bool GetSelectionLocation(FBox& OutLocation) const;
 
 	/** Whether there any foliage instances painted on specified component */
 	static FOLIAGE_API bool HasFoliageAttached(UActorComponent* InComponent);
@@ -261,18 +269,17 @@ private:
 #if WITH_EDITOR
 	void ClearSelection();
 	void OnLevelActorMoved(AActor* InActor);
+	void OnLevelActorOuterChanged(AActor* InActor, UObject* OldOuter);
 	void OnLevelActorDeleted(AActor* InActor);
 	void OnApplyLevelTransform(const FTransform& InTransform);
 	void OnPostApplyLevelOffset(ULevel* InLevel, UWorld* InWorld, const FVector& InOffset, bool bWorldShift);
 	void OnPostWorldInitialization(UWorld* World, const UWorld::InitializationValues IVS);
-
-	// Move instances to a foliage actor in target level
-	FOLIAGE_API void MoveInstancesToLevel(ULevel* InTargetLevel, TSet<int32>& InInstanceList, FFoliageInfo* InCurrentMeshInfo, UFoliageType* InFoliageType);
 #endif
 private:
 #if WITH_EDITOR
 	FDelegateHandle OnLevelActorMovedDelegateHandle;
 	FDelegateHandle OnLevelActorDeletedDelegateHandle;
+	FDelegateHandle OnLevelActorOuterChangedDelegateHandle;
 	FDelegateHandle OnPostApplyLevelOffsetDelegateHandle;
 	FDelegateHandle OnApplyLevelTransformDelegateHandle;
 	FDelegateHandle OnPostWorldInitializationDelegateHandle;
diff --git a/Engine/Source/Runtime/Landscape/Private/LandscapeSplines.cpp b/Engine/Source/Runtime/Landscape/Private/LandscapeSplines.cpp
--- a/Engine/Source/Runtime/Landscape/Private/LandscapeSplines.cpp
+++ b/Engine/Source/Runtime/Landscape/Private/LandscapeSplines.cpp
@@ -3233,6 +3233,11 @@ void ALandscapeProxy::CreateSplineComponent(const FVector& Scale3D)
 	SplineComponent->SetRelativeScale3D_Direct(Scale3D);
 	SplineComponent->AttachToComponent(GetRootComponent(), FAttachmentTransformRules::KeepRelativeTransform);
 	SplineComponent->ShowSplineEditorMesh(true);
+
+	if (RootComponent && RootComponent->IsRegistered())
+	{
+		SplineComponent->RegisterComponent();
+	}
 }
 
 void FindStartingControlPoints(ULandscapeSplineControlPoint* ControlPoint, TSet<ULandscapeSplineControlPoint*>& StartingControlPoints)
