Fix for non-determinism within serialized UFoliageTypes.

diff --git a/Engine/Source/Runtime/Foliage/Public/FoliageType.h b/Engine/Source/Runtime/Foliage/Public/FoliageType.h
--- a/Engine/Source/Runtime/Foliage/Public/FoliageType.h
+++ b/Engine/Source/Runtime/Foliage/Public/FoliageType.h
@@ -445,7 +445,7 @@ public:
 	UPROPERTY(Category = Procedural, EditAnywhere, meta = (Subcategory = "Growth", XAxisName = "Normalized Age", YAxisName = "Scale Factor"))
 	FRuntimeFloatCurve ScaleCurve;
 
-	UPROPERTY()
+	UPROPERTY(Transient)
 	int32 ChangeCount;
 
 public:
