Foliage handles Undo/Redo so skip event when GIsTransacting

#rnx

diff --git a/Engine/Source/Runtime/Foliage/Private/InstancedFoliage.cpp b/Engine/Source/Runtime/Foliage/Private/InstancedFoliage.cpp
--- a/Engine/Source/Runtime/Foliage/Private/InstancedFoliage.cpp
+++ b/Engine/Source/Runtime/Foliage/Private/InstancedFoliage.cpp
@@ -4003,6 +4003,11 @@ void AInstancedFoliageActor::OnLevelActorMoved(AActor* InActor)
 
 void AInstancedFoliageActor::OnLevelActorOuterChanged(AActor* InActor, UObject* OldOuter)
 {
+	if (GIsTransacting)
+	{
+		return;
+	}
+
 	ULevel* OldLevel = Cast<ULevel>(OldOuter);
 
 	if (InActor->GetLevel() == OldLevel)
