Fix Foliage "Snap to Floor" crash when snapping on foliage that doesn't have a valid base component

diff --git a/Engine/Source/Runtime/Foliage/Private/InstancedFoliage.cpp b/Engine/Source/Runtime/Foliage/Private/InstancedFoliage.cpp
--- a/Engine/Source/Runtime/Foliage/Private/InstancedFoliage.cpp
+++ b/Engine/Source/Runtime/Foliage/Private/InstancedFoliage.cpp
@@ -3252,6 +3252,12 @@ bool AInstancedFoliageActor::FoliageTrace(const UWorld* InWorld, FHitResult& Out
 					break;
 				}
 			}
+
+			// The foliage we are snapping on doesn't have a valid base
+			if (!OutHit.Component.IsValid())
+			{
+				continue; 
+			}
 		}
 
 		return bInsideProceduralVolumeOrArentUsingOne;
