Fix Undo of Delete Foliage Actor Instance (Add GarbageCollect so that refs to deleted construction script components are nulled out on the actor)

diff --git a/Engine/Source/Runtime/Foliage/Private/FoliageActor.cpp b/Engine/Source/Runtime/Foliage/Private/FoliageActor.cpp
--- a/Engine/Source/Runtime/Foliage/Private/FoliageActor.cpp
+++ b/Engine/Source/Runtime/Foliage/Private/FoliageActor.cpp
@@ -146,6 +146,7 @@ void FFoliageActor::RemoveInstance(int32 InstanceIndex)
 	AActor* Actor = ActorInstances[InstanceIndex];
 	ActorInstances.RemoveAtSwap(InstanceIndex);
 	Actor->GetWorld()->DestroyActor(Actor, true);
+	bActorsDestroyed = true;
 }
 
 void FFoliageActor::MoveInstance(int32 InstanceIndex, UObject*& OutInstanceImplementation)
@@ -155,6 +156,23 @@ void FFoliageActor::MoveInstance(int32 InstanceIndex, UObject*& OutInstanceImple
 	ActorInstances.RemoveAtSwap(InstanceIndex);
 }
 
+void FFoliageActor::BeginUpdate()
+{
+	bActorsDestroyed = false;
+}
+
+void FFoliageActor::EndUpdate()
+{
+	if (bActorsDestroyed)
+	{
+		// This is to null out refs to components that have been created through ConstructionScript (same as it is done in edactDeleteSelected).
+		// Because components that return true for IsCreatedByConstructionScript forward their Modify calls to their owning Actor so they are not part of the transaction.
+		// Undoing the DestroyActor will re-run the construction script and those components will be recreated.
+		CollectGarbage(GARBAGE_COLLECTION_KEEPFLAGS);
+	}
+	bActorsDestroyed = false;
+}
+
 void FFoliageActor::SetInstanceWorldTransform(int32 InstanceIndex, const FTransform& Transform, bool bTeleport)
 {
 	if (AActor* Actor = ActorInstances[InstanceIndex])
diff --git a/Engine/Source/Runtime/Foliage/Private/FoliageActor.h b/Engine/Source/Runtime/Foliage/Private/FoliageActor.h
--- a/Engine/Source/Runtime/Foliage/Private/FoliageActor.h
+++ b/Engine/Source/Runtime/Foliage/Private/FoliageActor.h
@@ -16,6 +16,9 @@ struct FFoliageActor : public FFoliageImpl
 	FFoliageActor()
 		: ActorClass(nullptr)
 		, bShouldAttachToBaseComponent(true)
+#if WITH_EDITOR
+		, bActorsDestroyed(false)
+#endif
 	{
 	}
 
@@ -44,6 +47,9 @@ struct FFoliageActor : public FFoliageImpl
 	virtual void SelectInstances(bool bSelect, const TSet<int32>& SelectedIndices) override;
 	virtual void ApplySelection(bool bApply, const TSet<int32>& SelectedIndices) override;
 	virtual void ClearSelection(const TSet<int32>& SelectedIndices) override;
+
+	virtual void BeginUpdate() override;
+	virtual void EndUpdate() override;
 	virtual void Refresh(AInstancedFoliageActor* IFA, const TArray<FFoliageInstance>& Instances, bool Async, bool Force) override;
 	virtual void OnHiddenEditorViewMaskChanged(uint64 InHiddenEditorViews) override;
 	virtual void PostEditUndo(AInstancedFoliageActor* IFA, UFoliageType* FoliageType, const TArray<FFoliageInstance>& Instances, const TSet<int32>& SelectedIndices) override;
@@ -59,5 +65,7 @@ struct FFoliageActor : public FFoliageImpl
 
 private:
 	void UpdateActorTransforms(const TArray<FFoliageInstance>& Instances);
+
+	bool bActorsDestroyed;
 #endif
 };
\ No newline at end of file
