Optimize FProceduralFoliageBroadphase::Insert by removing temporary string allocation that did also leave an invalid pointer.

diff --git a/Engine/Source/Runtime/Foliage/Private/ProceduralFoliageBroadphase.cpp b/Engine/Source/Runtime/Foliage/Private/ProceduralFoliageBroadphase.cpp
--- a/Engine/Source/Runtime/Foliage/Private/ProceduralFoliageBroadphase.cpp
+++ b/Engine/Source/Runtime/Foliage/Private/ProceduralFoliageBroadphase.cpp
@@ -38,7 +38,7 @@ bool FProceduralFoliageBroadphase::TestAgainstAABB(FProceduralFoliageInstance* I
 void FProceduralFoliageBroadphase::Insert(FProceduralFoliageInstance* Instance)
 {
 	const FBox2D MaxAABB = GetMaxAABB(Instance);
-	QuadTree.Insert(Instance, MaxAABB, Instance->Type ? *Instance->Type->GetName() : nullptr);
+	QuadTree.Insert(Instance, MaxAABB);
 }
 
 bool CircleOverlap(const FVector& ALocation, float ARadius, const FVector& BLocation, float BRadius)
