Foliage Actor: Allow UserConstructionScript to execute when Spawning Actors

diff --git a/Engine/Source/Runtime/Foliage/Private/FoliageActor.cpp b/Engine/Source/Runtime/Foliage/Private/FoliageActor.cpp
--- a/Engine/Source/Runtime/Foliage/Private/FoliageActor.cpp
+++ b/Engine/Source/Runtime/Foliage/Private/FoliageActor.cpp
@@ -81,7 +81,7 @@ AActor* FFoliageActor::Spawn(AInstancedFoliageActor* IFA, const FFoliageInstance
 	{
 		return nullptr;
 	}
-
+	FEditorScriptExecutionGuard ScriptGuard;
 	FActorSpawnParameters SpawnParameters;
 	SpawnParameters.ObjectFlags = RF_Transactional;
 	SpawnParameters.bHideFromSceneOutliner = true;
