Fix foliage non-deterministic cooking by removing unused UProceduralFoliageSpawner::bNeedsSimulation that can change depending on the cooking order.

diff --git a/Engine/Source/Runtime/Foliage/Private/ProceduralFoliageComponent.cpp b/Engine/Source/Runtime/Foliage/Private/ProceduralFoliageComponent.cpp
--- a/Engine/Source/Runtime/Foliage/Private/ProceduralFoliageComponent.cpp
+++ b/Engine/Source/Runtime/Foliage/Private/ProceduralFoliageComponent.cpp
@@ -126,7 +126,7 @@ bool UProceduralFoliageComponent::ExecuteSimulation(TArray<FDesiredFoliageInstan
 		FTileLayout TileLayout;
 		GetTileLayout(TileLayout);
 
-		FoliageSpawner->SimulateIfNeeded();
+		FoliageSpawner->Simulate();
 
 		TArray<TFuture< TArray<FDesiredFoliageInstance>* >> Futures;
 		for (int32 X = 0; X < TileLayout.NumTilesX; ++X)
diff --git a/Engine/Source/Runtime/Foliage/Private/ProceduralFoliageSpawner.cpp b/Engine/Source/Runtime/Foliage/Private/ProceduralFoliageSpawner.cpp
--- a/Engine/Source/Runtime/Foliage/Private/ProceduralFoliageSpawner.cpp
+++ b/Engine/Source/Runtime/Foliage/Private/ProceduralFoliageSpawner.cpp
@@ -17,13 +17,6 @@ UProceduralFoliageSpawner::UProceduralFoliageSpawner(const FObjectInitializer& O
 	RandomSeed = 42;
 }
 
-#if WITH_EDITOR
-void UProceduralFoliageSpawner::PostEditChangeProperty(FPropertyChangedEvent& PropertyChangedEvent)
-{
-	bNeedsSimulation = true;
-}
-#endif
-
 UProceduralFoliageTile* UProceduralFoliageSpawner::CreateTempTile()
 {
 	UProceduralFoliageTile* TmpTile = NewObject<UProceduralFoliageTile>(this);
@@ -47,8 +40,6 @@ void UProceduralFoliageSpawner::SetClean()
 	{
 		FoliageTypeObject.SetClean();
 	}
-
-	bNeedsSimulation = false;
 }
 
 // Custom serialization version for all packages containing ProceduralFoliage
@@ -83,34 +74,6 @@ void UProceduralFoliageSpawner::Serialize(FArchive& Ar)
 	Ar.UsingCustomVersion(FProceduralFoliageCustomVersion::GUID);
 }
 
-bool UProceduralFoliageSpawner::AnyDirty() const
-{
-	bool bDirty = bNeedsSimulation;
-
-	if (!bDirty)
-	{
-		for (const FFoliageTypeObject& FoliageTypeObject : FoliageTypes)
-		{
-			if (FoliageTypeObject.IsDirty())
-			{
-				bDirty = true;
-				break;
-			}
-		}
-	}
-	
-	return bDirty;
-}
-
-void UProceduralFoliageSpawner::SimulateIfNeeded()
-{
-	//if (AnyDirty())	@todo: for now we must force simulation every time since precomputed tiles are weak pointers
-	{
-		Simulate();
-	}
-
-}
-
 void UProceduralFoliageSpawner::Empty()
 {
 	for (TWeakObjectPtr<UProceduralFoliageTile>& WeakTile : PrecomputedTiles)
diff --git a/Engine/Source/Runtime/Foliage/Public/ProceduralFoliageSpawner.h b/Engine/Source/Runtime/Foliage/Public/ProceduralFoliageSpawner.h
--- a/Engine/Source/Runtime/Foliage/Public/ProceduralFoliageSpawner.h
+++ b/Engine/Source/Runtime/Foliage/Public/ProceduralFoliageSpawner.h
@@ -41,9 +41,6 @@ private:
 	UPROPERTY(Category = ProceduralFoliageSimulation, EditAnywhere)
 	TArray<FFoliageTypeObject> FoliageTypes;
 
-	UPROPERTY()
-	bool bNeedsSimulation;
-
 public:
 	UFUNCTION(BlueprintCallable, Category = ProceduralFoliageSimulation)
 	void Simulate(int32 NumSteps = -1);
@@ -56,15 +53,8 @@ public:
 	/** Returns the instances that need to spawn for a given min,max region */
 	void GetInstancesToSpawn(TArray<FProceduralFoliageInstance>& OutInstances, const FVector& Min = FVector(-FLT_MAX, -FLT_MAX, -FLT_MAX), const FVector& Max = FVector(FLT_MAX, FLT_MAX, FLT_MAX) ) const;
 
-#if WITH_EDITOR
-	virtual void PostEditChangeProperty(FPropertyChangedEvent& PropertyChangedEvent);
-#endif
-
 	virtual void Serialize(FArchive& Ar);
 
-	/** Simulates tiles if needed */
-	void SimulateIfNeeded();
-
 	/** Takes a tile index and returns a random tile associated with that index. */
 	const UProceduralFoliageTile* GetRandomTile(int32 X, int32 Y);
 
@@ -74,8 +64,6 @@ public:
 private:
 	void CreateProceduralFoliageInstances();
 
-	bool AnyDirty() const;
-
 	void SetClean();
 private:
 	TArray<TWeakObjectPtr<UProceduralFoliageTile>> PrecomputedTiles;
