Edigrated CL 13014082

Fixed crash in HISM. OnPostLoadPerInstanceData was properly detecting and calling a force BuildTree, but the build was failing because the static mesh of the HISM was null. For an unknown reason, data was corrupted but FoliageType's static mesh was still valid. We now detect this rare case and fix-up the component in AInstancedFoliageActor::PostLoad (editor-only).

diff --git a/Engine/Source/Runtime/Foliage/Private/InstancedFoliage.cpp b/Engine/Source/Runtime/Foliage/Private/InstancedFoliage.cpp
--- a/Engine/Source/Runtime/Foliage/Private/InstancedFoliage.cpp
+++ b/Engine/Source/Runtime/Foliage/Private/InstancedFoliage.cpp
@@ -3803,6 +3803,25 @@ void AInstancedFoliageActor::PostLoad()
 					}
 				}
 			}
+
+			// Fixup corrupted data
+			if (Info.Type == EFoliageImplType::StaticMesh)
+			{
+				UFoliageType_InstancedStaticMesh* FoliageType_InstancedStaticMesh = Cast<UFoliageType_InstancedStaticMesh>(FoliageType);
+				if (UStaticMesh* FoliageTypeStaticMesh = FoliageType_InstancedStaticMesh->GetStaticMesh())
+				{
+					FFoliageStaticMesh* FoliageStaticMesh = StaticCast<FFoliageStaticMesh*>(Info.Implementation.Get());
+					if (UHierarchicalInstancedStaticMeshComponent* HISMComponent = FoliageStaticMesh->Component)
+					{
+						HISMComponent->ConditionalPostLoad();
+						UStaticMesh* ComponentStaticMesh = HISMComponent->GetStaticMesh();
+						if (ComponentStaticMesh != FoliageTypeStaticMesh)
+						{
+							HISMComponent->SetStaticMesh(FoliageTypeStaticMesh);
+						}
+					}
+				}
+			}
 		}
 
 		UWorld* World = GetWorld();
