OFPA: Fix InstancedFoliageActor dirtying package on selection change

diff --git a/Engine/Source/Runtime/Foliage/Private/InstancedFoliage.cpp b/Engine/Source/Runtime/Foliage/Private/InstancedFoliage.cpp
--- a/Engine/Source/Runtime/Foliage/Private/InstancedFoliage.cpp
+++ b/Engine/Source/Runtime/Foliage/Private/InstancedFoliage.cpp
@@ -2185,7 +2185,7 @@ void FFoliageInfo::SelectInstances(AInstancedFoliageActor* InIFA, bool bSelect)
 {
 	if (Implementation->IsInitialized())
 	{
-		InIFA->Modify();
+		InIFA->Modify(false);
 
 		if (bSelect)
 		{
@@ -2215,7 +2215,7 @@ void FFoliageInfo::SelectInstances(AInstancedFoliageActor* InIFA, bool bSelect,
 		check(Implementation->IsInitialized());
 		if (bSelect)
 		{
-			InIFA->Modify();
+			InIFA->Modify(false);
 
 			SelectedIndices.Reserve(InInstances.Num());
 
@@ -2227,7 +2227,7 @@ void FFoliageInfo::SelectInstances(AInstancedFoliageActor* InIFA, bool bSelect,
 		}
 		else
 		{
-			InIFA->Modify();
+			InIFA->Modify(false);
 
 			for (int32 i : InInstances)
 			{
@@ -3169,7 +3169,7 @@ void AInstancedFoliageActor::ClearSelection()
 
 void AInstancedFoliageActor::SelectInstance(UInstancedStaticMeshComponent* InComponent, int32 InInstanceIndex, bool bToggle)
 {
-	Modify();
+	Modify(false);
 
 	// If we're not toggling, we need to first deselect everything else
 	if (!bToggle)
@@ -3253,7 +3253,7 @@ bool AInstancedFoliageActor::SelectInstance(AActor* InActor, bool bToggle)
 			return false;
 		}
 
-		Modify();
+		Modify(false);
 
 		// If we're not toggling, we need to first deselect everything else
 		if (!bToggle)
