Added Focus to Foliage Instances in Foliage Edit Mode

diff --git a/Engine/Source/Editor/FoliageEdit/Private/FoliageEdMode.cpp b/Engine/Source/Editor/FoliageEdit/Private/FoliageEdMode.cpp
--- a/Engine/Source/Editor/FoliageEdit/Private/FoliageEdMode.cpp
+++ b/Engine/Source/Editor/FoliageEdit/Private/FoliageEdMode.cpp
@@ -3,6 +3,7 @@
 #include "FoliageEdMode.h"
 #include "SceneView.h"
 #include "EditorViewportClient.h"
+#include "Editor.h"
 #include "Components/PrimitiveComponent.h"
 #include "Components/StaticMeshComponent.h"
 #include "FoliageType.h"
@@ -1763,6 +1764,18 @@ void FEdModeFoliage::SelectInstances(const UFoliageType* Settings, bool bSelect)
 	SelectInstances(GetWorld(), Settings, bSelect);
 }
 
+void FEdModeFoliage::FocusSelectedInstances() const
+{
+	UWorld* World = GetWorld();
+	FBox SelectionBoundingBox(EForceInit::ForceInit);
+	for (TActorIterator<AInstancedFoliageActor> It(World); It; ++It)
+	{
+		AInstancedFoliageActor* IFA = *It;
+		SelectionBoundingBox += IFA->GetSelectionBoundingBox();
+	}
+	GEditor->MoveViewportCamerasToBox(SelectionBoundingBox, true);
+}
+
 void FEdModeFoliage::SelectInstances(UWorld* InWorld, bool bSelect)
 {
 	FEdModeFoliageSelectionUpdate Scope(this);
@@ -3636,6 +3649,18 @@ bool FEdModeFoliage::InputKey(FEditorViewportClient* ViewportClient, FViewport*
 		}
 	}
 
+	if (!bHandled)
+	{
+		if (Event == IE_Pressed)
+		{
+			if (Key == EKeys::F)
+			{
+				FocusSelectedInstances();
+				bHandled = true;
+			}
+		}
+	}
+
 	return bHandled;
 }
 
diff --git a/Engine/Source/Editor/FoliageEdit/Private/FoliageEdMode.h b/Engine/Source/Editor/FoliageEdit/Private/FoliageEdMode.h
--- a/Engine/Source/Editor/FoliageEdit/Private/FoliageEdMode.h
+++ b/Engine/Source/Editor/FoliageEdit/Private/FoliageEdMode.h
@@ -502,6 +502,9 @@ public:
 	/** Set/Clear selection for foliage instances of specific type  */
 	void SelectInstances(const UFoliageType* Settings, bool bSelect);
 
+	/*Focus on selected instances*/
+	void FocusSelectedInstances() const;
+
 	/** Find and select instances that don't have valid base or 'off-ground' */
 	void SelectInvalidInstances(const TArray<const UFoliageType*>& FoliageTypes);
 
diff --git a/Engine/Source/Runtime/Foliage/Private/FoliageActor.cpp b/Engine/Source/Runtime/Foliage/Private/FoliageActor.cpp
--- a/Engine/Source/Runtime/Foliage/Private/FoliageActor.cpp
+++ b/Engine/Source/Runtime/Foliage/Private/FoliageActor.cpp
@@ -334,6 +334,16 @@ void FFoliageActor::SelectInstances(bool bSelect, const TSet<int32>& SelectedInd
 	AInstancedFoliageActor::SelectionChanged.Broadcast(bSelect, GetActorsFromSelectedIndices(SelectedIndices));
 }
 
+FBox FFoliageActor::GetSelectionBoundingBox(const TSet<int32>& SelectedIndices) const
+{
+	FBox BoundingBox(EForceInit::ForceInit);
+	TArray<AActor*> SelectedActors = GetActorsFromSelectedIndices(SelectedIndices);
+	for (auto Actor : SelectedActors)
+	{
+		BoundingBox += Actor->GetComponentsBoundingBox();
+	}
+	return BoundingBox;
+}
 void FFoliageActor::ApplySelection(bool bApply, const TSet<int32>& SelectedIndices)
 {
 	if (bApply && SelectedIndices.Num() > 0)
diff --git a/Engine/Source/Runtime/Foliage/Private/FoliageActor.h b/Engine/Source/Runtime/Foliage/Private/FoliageActor.h
--- a/Engine/Source/Runtime/Foliage/Private/FoliageActor.h
+++ b/Engine/Source/Runtime/Foliage/Private/FoliageActor.h
@@ -45,6 +45,7 @@ struct FFoliageActor : public FFoliageImpl
 	virtual void SelectAllInstances(bool bSelect) override;
 	virtual void SelectInstance(bool bSelect, int32 Index) override;
 	virtual void SelectInstances(bool bSelect, const TSet<int32>& SelectedIndices) override;
+	virtual FBox GetSelectionBoundingBox(const TSet<int32>& SelectedIndices) const override;
 	virtual void ApplySelection(bool bApply, const TSet<int32>& SelectedIndices) override;
 	virtual void ClearSelection(const TSet<int32>& SelectedIndices) override;
 
diff --git a/Engine/Source/Runtime/Foliage/Private/InstancedFoliage.cpp b/Engine/Source/Runtime/Foliage/Private/InstancedFoliage.cpp
--- a/Engine/Source/Runtime/Foliage/Private/InstancedFoliage.cpp
+++ b/Engine/Source/Runtime/Foliage/Private/InstancedFoliage.cpp
@@ -116,6 +116,7 @@ struct FFoliageStaticMesh : public FFoliageImpl
 	virtual void SelectAllInstances(bool bSelect) override;
 	virtual void SelectInstance(bool bSelect, int32 Index) override;
 	virtual void SelectInstances(bool bSelect, const TSet<int32>& SelectedIndices) override;
+	virtual FBox GetSelectionBoundingBox(const TSet<int32>& SelectedIndices) const override;
 	virtual void ApplySelection(bool bApply, const TSet<int32>& SelectedIndices) override;
 	virtual void ClearSelection(const TSet<int32>& SelectedIndices) override;
 
@@ -1045,6 +1046,18 @@ void FFoliageStaticMesh::SelectInstances(bool bSelect, const TSet<int32>& Select
 	Component->MarkRenderStateDirty();
 }
 
+FBox FFoliageStaticMesh::GetSelectionBoundingBox(const TSet<int32>& SelectedIndices) const
+{
+	FBox BoundingBox(EForceInit::ForceInit);
+	for (int32 i : SelectedIndices)
+	{
+		FTransform InstanceWorldTransform;
+		Component->GetInstanceTransform(i, InstanceWorldTransform, true);
+		BoundingBox += Component->GetStaticMesh()->GetBoundingBox().TransformBy(InstanceWorldTransform);
+	}
+	return BoundingBox;
+}
+
 void FFoliageStaticMesh::ApplySelection(bool bApply, const TSet<int32>& SelectedIndices)
 {
 	if (Component && (bApply || Component->SelectedInstances.Num() > 0))
@@ -2206,6 +2219,12 @@ void FFoliageInfo::SelectInstances(AInstancedFoliageActor* InIFA, bool bSelect)
 	}
 }
 
+FBox FFoliageInfo::GetSelectionBoundingBox() const
+{
+	check(Implementation->IsInitialized());
+	return Implementation->GetSelectionBoundingBox(SelectedIndices);
+}
+
 void FFoliageInfo::SelectInstances(AInstancedFoliageActor* InIFA, bool bSelect, TArray<int32>& InInstances)
 {
 	if (InInstances.Num())
@@ -3281,6 +3300,19 @@ bool AInstancedFoliageActor::SelectInstance(AActor* InActor, bool bToggle)
 	return true;
 }
 
+FOLIAGE_API FBox AInstancedFoliageActor::GetSelectionBoundingBox() const
+{
+	FBox SelectionBoundingBox(EForceInit::ForceInit);
+
+	for (auto& Pair : FoliageInfos)
+	{
+		const FFoliageInfo* Info = &Pair.Value.Get();
+		SelectionBoundingBox += Info->GetSelectionBoundingBox();
+	}
+
+	return SelectionBoundingBox;
+}
+
 bool AInstancedFoliageActor::HasSelectedInstances() const
 {
 	for (const auto& Pair : FoliageInfos)
diff --git a/Engine/Source/Runtime/Foliage/Public/InstancedFoliage.h b/Engine/Source/Runtime/Foliage/Public/InstancedFoliage.h
--- a/Engine/Source/Runtime/Foliage/Public/InstancedFoliage.h
+++ b/Engine/Source/Runtime/Foliage/Public/InstancedFoliage.h
@@ -200,6 +200,7 @@ struct FFoliageImpl
 	virtual void SelectAllInstances(bool bSelect) = 0;
 	virtual void SelectInstance(bool bSelect, int32 Index) = 0;
 	virtual void SelectInstances(bool bSelect, const TSet<int32>& SelectedIndices) = 0;
+	virtual FBox GetSelectionBoundingBox(const TSet<int32>& SelectedIndices) const = 0;
 	virtual void ApplySelection(bool bApply, const TSet<int32>& SelectedIndices) = 0;
 	virtual void ClearSelection(const TSet<int32>& SelectedIndices) = 0;
 
@@ -310,6 +311,8 @@ struct FFoliageInfo
 
 	FOLIAGE_API void SelectInstances(AInstancedFoliageActor* InIFA, bool bSelect);
 
+	FOLIAGE_API FBox GetSelectionBoundingBox() const;
+
 	// Get the number of placed instances
 	FOLIAGE_API int32 GetPlacedInstanceCount() const;
 
diff --git a/Engine/Source/Runtime/Foliage/Public/InstancedFoliageActor.h b/Engine/Source/Runtime/Foliage/Public/InstancedFoliageActor.h
--- a/Engine/Source/Runtime/Foliage/Public/InstancedFoliageActor.h
+++ b/Engine/Source/Runtime/Foliage/Public/InstancedFoliageActor.h
@@ -243,6 +243,9 @@ public:
 	// Select an individual instance.
 	FOLIAGE_API bool SelectInstance(AActor* InActor, bool bToggle);
 
+	//Get the bounds of all selected instances
+	FOLIAGE_API FBox GetSelectionBoundingBox() const;
+
 	// Whether actor has selected instances
 	FOLIAGE_API bool HasSelectedInstances() const;
 
