WorldPartiton: Foliage remove CastChecked to avoid crashing on null Actors (can happen with unloaded FoliageActors)

diff --git a/Engine/Source/Runtime/Foliage/Private/InstancedFoliage.cpp b/Engine/Source/Runtime/Foliage/Private/InstancedFoliage.cpp
--- a/Engine/Source/Runtime/Foliage/Private/InstancedFoliage.cpp
+++ b/Engine/Source/Runtime/Foliage/Private/InstancedFoliage.cpp
@@ -2526,7 +2526,7 @@ AInstancedFoliageActor* AInstancedFoliageActor::Get(UWorld* InWorld, bool bCreat
 {
 	UActorPartitionSubsystem* ActorPartitionSubsystem = InWorld->GetSubsystem<UActorPartitionSubsystem>();
 
-	return CastChecked<AInstancedFoliageActor>(
+	return Cast<AInstancedFoliageActor>(
 		ActorPartitionSubsystem->GetActor(
 			FActorPartitionGetParams(
 				AInstancedFoliageActor::StaticClass(), 
