Foliage optim: avoid constructing a temporary FSoftObjectPath during painting.

diff --git a/Engine/Source/Runtime/Foliage/Private/FoliageInstanceBase.cpp b/Engine/Source/Runtime/Foliage/Private/FoliageInstanceBase.cpp
--- a/Engine/Source/Runtime/Foliage/Private/FoliageInstanceBase.cpp
+++ b/Engine/Source/Runtime/Foliage/Private/FoliageInstanceBase.cpp
@@ -166,14 +166,8 @@ FFoliageInstanceBaseId FFoliageInstanceBaseCache::AddInstanceBaseId(UActorCompon
 
 FFoliageInstanceBaseId FFoliageInstanceBaseCache::GetInstanceBaseId(UActorComponent* InComponent) const
 {
-	if(InstanceBaseInvMap.Num() > 0)
+	if(InComponent && (InstanceBaseInvMap.Num() > 0))
 	{
-		// test if this component has allocated FSoftObjectPath, to avoid creating new one in FFoliageInstanceBasePtr ctor
-		if (!FSoftObjectPath(InComponent).IsValid())
-		{
-			return InvalidBaseId;
-		}
-
 		FFoliageInstanceBasePtr BasePtr(InComponent);
 		if (BasePtr.IsValid())
 		{
