In World Partition, load editor cells intersecting ProceduralFoliageCompoment bounds before simulation.

diff --git a/Engine/Source/Runtime/Foliage/Private/ProceduralFoliageComponent.cpp b/Engine/Source/Runtime/Foliage/Private/ProceduralFoliageComponent.cpp
--- a/Engine/Source/Runtime/Foliage/Private/ProceduralFoliageComponent.cpp
+++ b/Engine/Source/Runtime/Foliage/Private/ProceduralFoliageComponent.cpp
@@ -12,6 +12,10 @@
 #include "EngineUtils.h"
 #include "Misc/FeedbackContext.h"
 
+#if WITH_EDITOR
+#include "WorldPartition/WorldPartition.h"
+#endif
+
 
 #define LOCTEXT_NAMESPACE "ProceduralFoliage"
 
@@ -113,6 +117,15 @@ FVector UProceduralFoliageComponent::GetWorldPosition() const
 bool UProceduralFoliageComponent::ExecuteSimulation(TArray<FDesiredFoliageInstance>& OutInstances)
 {
 #if WITH_EDITOR
+
+	// In World Partition, load Editor Cells intersecting bounds affected by ProceduralFoliageCompoment
+	if (UWorldPartition* WorldPartition = GetWorld()->GetWorldPartition())
+	{
+		FBox Bounds = GetBounds();
+		check(Bounds.IsValid);
+		WorldPartition->LoadEditorCells(Bounds);
+	}
+
 	FBodyInstance* BoundsBodyInstance = GetBoundsBodyInstance();
 	if (FoliageSpawner)
 	{
