Added UFoliageType::bIncludeInHLOD property

diff --git a/Engine/Source/Runtime/Core/Public/UObject/UE5MainStreamObjectVersion.h b/Engine/Source/Runtime/Core/Public/UObject/UE5MainStreamObjectVersion.h
--- a/Engine/Source/Runtime/Core/Public/UObject/UE5MainStreamObjectVersion.h
+++ b/Engine/Source/Runtime/Core/Public/UObject/UE5MainStreamObjectVersion.h
@@ -70,6 +70,9 @@ struct CORE_API FUE5MainStreamObjectVersion
 		// Fixup FLandscapeActorDesc Grid indices
 		FLandscapeActorDescFixupGridIndices,
 
+		// FoliageType with HLOD support
+		FoliageTypeIncludeInHLOD,
+
 		// -----<new versions can be added above this line>-------------------------------------------------
 		VersionPlusOne,
 		LatestVersion = VersionPlusOne - 1
diff --git a/Engine/Source/Runtime/Foliage/Private/InstancedFoliage.cpp b/Engine/Source/Runtime/Foliage/Private/InstancedFoliage.cpp
--- a/Engine/Source/Runtime/Foliage/Private/InstancedFoliage.cpp
+++ b/Engine/Source/Runtime/Foliage/Private/InstancedFoliage.cpp
@@ -42,6 +42,7 @@ InstancedFoliage.cpp: Instanced foliage implementation.
 #include "EngineGlobals.h"
 #include "Engine/StaticMesh.h"
 #include "DrawDebugHelpers.h"
+#include "UObject/UE5MainStreamObjectVersion.h"
 #include "UObject/FortniteMainBranchObjectVersion.h"
 #include "PreviewScene.h"
 #include "FoliageActor.h"
@@ -468,6 +469,10 @@ UFoliageType::UFoliageType(const FObjectInitializer& ObjectInitializer)
 	bEnableDensityScaling = false;
 	bEnableDiscardOnLoad = false;
 
+#if WITH_EDITORONLY_DATA
+	bIncludeInHLOD = true;
+#endif
+
 #if WITH_EDITORONLY_DATA
 	// Deprecated since FFoliageCustomVersion::FoliageTypeCustomization
 	ScaleMinX_DEPRECATED = 1.0f;
@@ -495,6 +500,7 @@ void UFoliageType::Serialize(FArchive& Ar)
 	Super::Serialize(Ar);
 
 	Ar.UsingCustomVersion(FFoliageCustomVersion::GUID);
+	Ar.UsingCustomVersion(FUE5MainStreamObjectVersion::GUID);
 
 	// we now have mask configurations for every color channel
 	if (Ar.IsLoading() && Ar.IsPersistent() && !Ar.HasAnyPortFlags(PPF_Duplicate | PPF_DuplicateForPIE) && VertexColorMask_DEPRECATED != FOLIAGEVERTEXCOLORMASK_Disabled)
@@ -579,6 +585,11 @@ void UFoliageType::Serialize(FArchive& Ar)
 			ProceduralScale.Min = MinScale_DEPRECATED;
 			ProceduralScale.Max = MaxScale_DEPRECATED;
 		}
+
+		if (Ar.CustomVer(FUE5MainStreamObjectVersion::GUID) < FUE5MainStreamObjectVersion::FoliageTypeIncludeInHLOD)
+		{
+			bIncludeInHLOD = false;
+		}
 	}
 #endif// WITH_EDITORONLY_DATA
 }
@@ -1469,6 +1480,12 @@ void FFoliageStaticMesh::UpdateComponentSettings(const UFoliageType_InstancedSta
 			bNeedsMarkRenderStateDirty = true;
 		}
 
+		if (Component->bEnableAutoLODGeneration != FoliageType->bIncludeInHLOD)
+		{
+			Component->bEnableAutoLODGeneration = FoliageType->bIncludeInHLOD;
+			bNeedsMarkRenderStateDirty = true;
+		}
+
 		const UFoliageType_InstancedStaticMesh* FoliageType_ISM = Cast<UFoliageType_InstancedStaticMesh>(FoliageType);
 		if (FoliageType_ISM)
 		{
diff --git a/Engine/Source/Runtime/Foliage/Public/FoliageType.h b/Engine/Source/Runtime/Foliage/Public/FoliageType.h
--- a/Engine/Source/Runtime/Foliage/Public/FoliageType.h
+++ b/Engine/Source/Runtime/Foliage/Public/FoliageType.h
@@ -554,6 +554,14 @@ public:
 	UPROPERTY(EditAnywhere, BlueprintReadOnly, Category = VirtualTexture, meta = (DisplayName = "Draw in Main Pass"))
 	ERuntimeVirtualTextureMainPassType VirtualTextureRenderPassType = ERuntimeVirtualTextureMainPassType::Exclusive;
 
+public:
+	// HLOD
+
+#if WITH_EDITORONLY_DATA
+	UPROPERTY(EditAnywhere, BlueprintReadWrite, Category=HLOD)
+	uint32 bIncludeInHLOD:1;
+#endif // WITH_EDITORONLY_DATA
+
 private:
 
 #if WITH_EDITORONLY_DATA
