Added Density Falloff curve to FoliageType (only available with ProceduralFoliageVolumes) and also to ProceduralFoliageBlockingVolume. This allows for smooth falloff of generated foliage instances on the edges of the volumes.

diff --git a/Engine/Source/Runtime/Foliage/Private/InstancedFoliage.cpp b/Engine/Source/Runtime/Foliage/Private/InstancedFoliage.cpp
--- a/Engine/Source/Runtime/Foliage/Private/InstancedFoliage.cpp
+++ b/Engine/Source/Runtime/Foliage/Private/InstancedFoliage.cpp
@@ -387,6 +387,56 @@ FArchive& operator<<(FArchive& Ar, FFoliageInfo& Info)
 	return Ar;
 }
 
+//
+// FFoliageDensityFalloff
+//
+
+FFoliageDensityFalloff::FFoliageDensityFalloff()
+{
+	FRichCurve* FalloffRichCurve = FalloffCurve.GetRichCurve();
+	FalloffRichCurve->AddKey(0.f, 1.f);
+	FalloffRichCurve->AddKey(1.f, 0.f);
+}
+
+bool FFoliageDensityFalloff::IsInstanceFiltered(const FVector2D& InstancePosition, const FVector2D& Origin, float MaxDistance) const
+{
+	float KeepPointProbability = GetDensityFalloffValue(InstancePosition, Origin, MaxDistance);
+	check(KeepPointProbability >= 0.f && KeepPointProbability <= 1.f);
+	if (KeepPointProbability < 1.f)
+	{
+		int32 PointSeed = GetRandomSeedForPosition(InstancePosition);
+		FRandomStream LocalRandomStream(PointSeed);
+		float Rand = LocalRandomStream.FRand();
+		return (Rand > KeepPointProbability);
+	}
+	return false;
+}
+
+float FFoliageDensityFalloff::GetDensityFalloffValue(const FVector2D& Position, const FVector2D& Origin, float MaxDistance) const
+{
+	float KeepPointProbability = 1.f;
+	if (bUseFalloffCurve)
+	{
+		float Distance = FVector2D::Distance(Position, Origin);
+		float NormalizedDistance = MaxDistance > 0.f ? (Distance / MaxDistance) : 1.f;
+		if (NormalizedDistance > 1.f)
+			NormalizedDistance = 1.f;
+		const FRichCurve* FalloffRichCurve = FalloffCurve.GetRichCurveConst();
+		KeepPointProbability = FMath::Clamp(FalloffRichCurve->Eval(NormalizedDistance), 0.f, 1.f);
+	}
+	return KeepPointProbability;
+}
+
+int32 FFoliageDensityFalloff::GetRandomSeedForPosition(const FVector2D& Position)
+{
+	// generate a unique random seed for a given position (precision = cm)
+	int32 Xcm = FMath::RoundToInt(Position.X);
+	int32 Ycm = FMath::RoundToInt(Position.Y);
+	// use the int32 hashing function to avoid patterns by spreading out distribution : 
+	return HashCombine(GetTypeHash(Xcm), GetTypeHash(Ycm));
+}
+
+
 //
 // UFoliageType
 //
@@ -4521,7 +4571,22 @@ bool AInstancedFoliageActor::FoliageTrace(const UWorld* InWorld, FHitResult& Out
 				const AProceduralFoliageVolume* ProceduralFoliageVolume = ProceduralFoliageBlockingVolume->ProceduralFoliageVolume;
 				if (ProceduralFoliageVolume == nullptr || ProceduralFoliageVolume->ProceduralComponent == nullptr || ProceduralFoliageVolume->ProceduralComponent->GetProceduralGuid() == DesiredInstance.ProceduralGuid)
 				{
-					return false;
+					if (!ProceduralFoliageBlockingVolume->DensityFalloff.bUseFalloffCurve)
+					{
+						return false;
+					}
+					else if (UBrushComponent* Brush = ProceduralFoliageBlockingVolume->GetBrushComponent())
+					{
+						FBox ActorVolumeBounds = Brush->Bounds.GetBox();
+						FVector2D ActorVolumeLocation = FVector2D(ActorVolumeBounds.GetCenter());
+						const float ActorVolumeMaxExtent = FVector2D(ActorVolumeBounds.GetExtent()).GetMax();
+
+						const FVector2D Origin(ProceduralFoliageBlockingVolume->GetActorTransform().GetLocation());
+						if (ProceduralFoliageBlockingVolume->DensityFalloff.IsInstanceFiltered(FVector2D(Hit.ImpactPoint), ActorVolumeLocation, ActorVolumeMaxExtent))
+						{
+							return false;
+						}
+					}
 				}
 			}
 			else if (HitObjectHandle.IsValid() && HitObjectHandle.DoesRepresentClass(AProceduralFoliageVolume::StaticClass())) //we never want to collide with our spawning volume
diff --git a/Engine/Source/Runtime/Foliage/Private/ProceduralFoliageComponent.cpp b/Engine/Source/Runtime/Foliage/Private/ProceduralFoliageComponent.cpp
--- a/Engine/Source/Runtime/Foliage/Private/ProceduralFoliageComponent.cpp
+++ b/Engine/Source/Runtime/Foliage/Private/ProceduralFoliageComponent.cpp
@@ -137,6 +137,9 @@ bool UProceduralFoliageComponent::ExecuteSimulation(TArray<FDesiredFoliageInstan
 		// Establish basic info about the tiles
 		const float TileSize = FoliageSpawner->TileSize;
 		const FVector WorldPosition = GetWorldPosition();
+		FBox ActorVolumeBounds = GetBounds();
+		FVector2D ActorVolumeLocation = FVector2D(ActorVolumeBounds.GetCenter());
+		const float ActorVolumeMaxExtent = FVector2D(ActorVolumeBounds.GetExtent()).GetMax();
 		FTileLayout TileLayout;
 		GetTileLayout(TileLayout);
 
@@ -206,7 +209,7 @@ bool UProceduralFoliageComponent::ExecuteSimulation(TArray<FDesiredFoliageInstan
 					const FTransform TileTM(OrientedOffset + WorldPosition);
 
 					TArray<FDesiredFoliageInstance>* DesiredInstances = new TArray<FDesiredFoliageInstance>();
-					CompositeTile->ExtractDesiredInstances(*DesiredInstances, TileTM, ProceduralGuid, TileLayout.HalfHeight, BoundsBodyInstance, true);
+					CompositeTile->ExtractDesiredInstances(*DesiredInstances, TileTM, ActorVolumeLocation, ActorVolumeMaxExtent, ProceduralGuid, TileLayout.HalfHeight, BoundsBodyInstance, true);
 
 					return DesiredInstances;
 					
diff --git a/Engine/Source/Runtime/Foliage/Private/ProceduralFoliageTile.cpp b/Engine/Source/Runtime/Foliage/Private/ProceduralFoliageTile.cpp
--- a/Engine/Source/Runtime/Foliage/Private/ProceduralFoliageTile.cpp
+++ b/Engine/Source/Runtime/Foliage/Private/ProceduralFoliageTile.cpp
@@ -465,7 +465,7 @@ void UProceduralFoliageTile::BeginDestroy()
 	RemoveInstances();
 }
 
-void UProceduralFoliageTile::ExtractDesiredInstances(TArray<FDesiredFoliageInstance>& OutInstances, const FTransform& WorldTM, const FGuid& ProceduralGuid, const float HalfHeight, const FBodyInstance* VolumeBodyInstance, bool bEmptyTileInfo)
+void UProceduralFoliageTile::ExtractDesiredInstances(TArray<FDesiredFoliageInstance>& OutInstances, const FTransform& WorldTM, const FVector2D& ActorVolumeLocation, float ActorVolumeMaxExtent, const FGuid& ProceduralGuid, const float HalfHeight, const FBodyInstance* VolumeBodyInstance, bool bEmptyTileInfo)
 {
 	InstancesToArray();
 
@@ -480,6 +480,13 @@ void UProceduralFoliageTile::ExtractDesiredInstances(TArray<FDesiredFoliageInsta
 		FVector EndRay = StartRay;
 		EndRay.Z -= (HalfHeight*2.f + 10.f);	//add 10cm to bottom position of raycast. This is needed because volume is usually placed directly on geometry and then you get precision issues
 
+		// Apply FoliageType's Falloff
+		FVector2D Position(StartRay);
+		if (Instance.Type->DensityFalloff.IsInstanceFiltered(Position, ActorVolumeLocation, ActorVolumeMaxExtent))
+		{
+			continue;
+		}
+		
 		FDesiredFoliageInstance* DesiredInst = new (OutInstances)FDesiredFoliageInstance(StartRay, EndRay, Instance.GetMaxRadius());
 		DesiredInst->Rotation = Instance.Rotation;
 		DesiredInst->ProceduralGuid = ProceduralGuid;
diff --git a/Engine/Source/Runtime/Foliage/Public/FoliageType.h b/Engine/Source/Runtime/Foliage/Public/FoliageType.h
--- a/Engine/Source/Runtime/Foliage/Public/FoliageType.h
+++ b/Engine/Source/Runtime/Foliage/Public/FoliageType.h
@@ -79,6 +79,31 @@ enum class EFoliageScaling : uint8
 	LockYZ
 };
 
+USTRUCT()
+struct FFoliageDensityFalloff
+{
+	GENERATED_USTRUCT_BODY()
+
+	FFoliageDensityFalloff();
+
+	UPROPERTY(Category = Procedural, EditAnywhere, meta = (Subcategory = "Density"))
+	bool bUseFalloffCurve = false;
+
+	/**
+	 * Density as a function of normalized distance (i.e. distance from Procedural Foliage Volume / Max Volume Extent).
+	 * X = 0 corresponds to Normalized distance = 0, X = 1 corresponds to Normalized distance = Max distance.
+	 * Y = 0 corresponds to 0% probability of keeping instance, Y = 1 corresponds to 100% probability of keeping instance.
+	 */
+	UPROPERTY(Category = Procedural, EditAnywhere, meta = (Subcategory = "Density", XAxisName = "Normalized Distance", YAxisName = "Density Factor"))
+	FRuntimeFloatCurve FalloffCurve;
+
+	FOLIAGE_API bool IsInstanceFiltered(const FVector2D& Position, const FVector2D& Origin, float MaxDistance) const;
+	FOLIAGE_API float GetDensityFalloffValue(const FVector2D& Position, const FVector2D& Origin, float MaxDistance) const;
+
+private:
+
+	static int32 GetRandomSeedForPosition(const FVector2D& Position);
+};
 
 UCLASS(hidecategories = Object, editinlinenew, MinimalAPI, BlueprintType, Blueprintable)
 class UFoliageType : public UObject
@@ -449,6 +474,9 @@ public:
 	UPROPERTY(Category = Procedural, EditAnywhere, meta = (Subcategory = "Growth", XAxisName = "Normalized Age", YAxisName = "Scale Factor"))
 	FRuntimeFloatCurve ScaleCurve;
 
+	UPROPERTY(Category = Procedural, EditAnywhere)
+	FFoliageDensityFalloff DensityFalloff;
+
 	UPROPERTY(Transient)
 	int32 ChangeCount;
 
diff --git a/Engine/Source/Runtime/Foliage/Public/ProceduralFoliageBlockingVolume.h b/Engine/Source/Runtime/Foliage/Public/ProceduralFoliageBlockingVolume.h
--- a/Engine/Source/Runtime/Foliage/Public/ProceduralFoliageBlockingVolume.h
+++ b/Engine/Source/Runtime/Foliage/Public/ProceduralFoliageBlockingVolume.h
@@ -6,6 +6,7 @@
 #include "CoreMinimal.h"
 #include "UObject/ObjectMacros.h"
 #include "GameFramework/Volume.h"
+#include "FoliageType.h"
 #include "ProceduralFoliageBlockingVolume.generated.h"
 
 class AProceduralFoliageVolume;
@@ -18,6 +19,9 @@ class AProceduralFoliageBlockingVolume : public AVolume
 
 	UPROPERTY(Category = ProceduralFoliage, EditAnywhere)
 	AProceduralFoliageVolume* ProceduralFoliageVolume;
+
+	UPROPERTY(Category = ProceduralFoliage, EditAnywhere)
+	FFoliageDensityFalloff DensityFalloff;
 };
 
 
diff --git a/Engine/Source/Runtime/Foliage/Public/ProceduralFoliageTile.h b/Engine/Source/Runtime/Foliage/Public/ProceduralFoliageTile.h
--- a/Engine/Source/Runtime/Foliage/Public/ProceduralFoliageTile.h
+++ b/Engine/Source/Runtime/Foliage/Public/ProceduralFoliageTile.h
@@ -42,7 +42,7 @@ class FOLIAGE_API UProceduralFoliageTile : public UObject
 	/**
 	 *	Extracts information on each instance that should be spawned based on the simulation.
 	 */
-	void ExtractDesiredInstances(TArray<FDesiredFoliageInstance>& OutDesiredInstances, const FTransform& WorldTM, const FGuid& ProceduralGuid, const float HalfHeight, const FBodyInstance* VolumeBodyInstance, bool bEmptyTileInfo = true);
+	void ExtractDesiredInstances(TArray<FDesiredFoliageInstance>& OutDesiredInstances, const FTransform& WorldTM, const FVector2D& ActorVolumeLocation, float ActorVolumeMaxExtent, const FGuid& ProceduralGuid, const float HalfHeight, const FBodyInstance* VolumeBodyInstance, bool bEmptyTileInfo = true);
 	
 	/**
 	 *	Copies instances within this tile to another
