Foliage: AddInstances / RemoveAllInstances Blueprint functions

diff --git a/Engine/Source/Runtime/Foliage/Private/InstancedFoliage.cpp b/Engine/Source/Runtime/Foliage/Private/InstancedFoliage.cpp
--- a/Engine/Source/Runtime/Foliage/Private/InstancedFoliage.cpp
+++ b/Engine/Source/Runtime/Foliage/Private/InstancedFoliage.cpp
@@ -4248,6 +4248,51 @@ void AInstancedFoliageActor::RemoveBaseComponentOnFoliageTypeInstances(UFoliageT
 	}
 }
 
+void AInstancedFoliageActor::AddInstances(UObject* WorldContextObject, UFoliageType* InFoliageType, TArray<FTransform> InTransforms)
+{
+	UWorld* World = GEngine->GetWorldFromContextObject(WorldContextObject, EGetWorldErrorMode::LogAndReturnNull);
+	if (World)
+	{
+		TMap<AInstancedFoliageActor*, TArray<const FFoliageInstance*>> InstancesToAdd;
+		TArray<FFoliageInstance> FoliageInstances;
+		FoliageInstances.Reserve(InTransforms.Num()); // Reserve 
+
+		for (const FTransform& InstanceTransfo : InTransforms)
+		{
+			AInstancedFoliageActor* IFA = AInstancedFoliageActor::Get(World, true, World->PersistentLevel, InstanceTransfo.GetLocation());
+			FFoliageInstance FoliageInstance;
+			FoliageInstance.Location = InstanceTransfo.GetLocation();
+			FoliageInstance.Rotation = InstanceTransfo.GetRotation().Rotator();
+			FoliageInstance.DrawScale3D = InstanceTransfo.GetScale3D();
+
+			FoliageInstances.Add(FoliageInstance);
+			InstancesToAdd.FindOrAdd(IFA).Add(&FoliageInstances[FoliageInstances.Num() - 1]);
+		}
+
+		for (const auto& Pair : InstancesToAdd)
+		{
+			FFoliageInfo* TypeInfo = nullptr;
+			if (UFoliageType* FoliageType = Pair.Key->AddFoliageType(InFoliageType, &TypeInfo))
+			{
+				TypeInfo->AddInstances(Pair.Key, FoliageType, Pair.Value);
+			}
+		}
+	}
+}
+
+void AInstancedFoliageActor::RemoveAllInstances(UObject* WorldContextObject, UFoliageType* InFoliageType)
+{
+	UWorld* World = GEngine->GetWorldFromContextObject(WorldContextObject, EGetWorldErrorMode::LogAndReturnNull);
+	if (World)
+	{
+		for (TActorIterator<AInstancedFoliageActor> It(World); It; ++It)
+		{
+			AInstancedFoliageActor* IFA = (*It);
+			IFA->RemoveFoliageType(&InFoliageType, 1);
+		}
+	}
+}
+
 void AInstancedFoliageActor::NotifyFoliageTypeWillChange(UFoliageType* FoliageType)
 {
 	FFoliageInfo* TypeInfo = FindInfo(FoliageType);
diff --git a/Engine/Source/Runtime/Foliage/Public/InstancedFoliageActor.h b/Engine/Source/Runtime/Foliage/Public/InstancedFoliageActor.h
--- a/Engine/Source/Runtime/Foliage/Public/InstancedFoliageActor.h
+++ b/Engine/Source/Runtime/Foliage/Public/InstancedFoliageActor.h
@@ -232,6 +232,12 @@ public:
 	void RepairDuplicateIFA(AInstancedFoliageActor* InDuplicateIFA);
 
 	void RemoveBaseComponentOnFoliageTypeInstances(UFoliageType* FoliageType);
+
+	UFUNCTION(BlueprintCallable, Category="Foliage", meta = (WorldContext = "WorldContextObject"))
+	static void AddInstances(UObject* WorldContextObject, UFoliageType* InFoliageType, TArray<FTransform> InTransforms);
+
+	UFUNCTION(BlueprintCallable, Category="Foliage", meta = (WorldContext = "WorldContextObject"))
+	static void RemoveAllInstances(UObject* WorldContextObject, UFoliageType* InFoliageType);
 #endif	//WITH_EDITOR
 
 private:
