Foliage: Change signature of FTransform array (const &)

diff --git a/Engine/Source/Runtime/Foliage/Private/InstancedFoliage.cpp b/Engine/Source/Runtime/Foliage/Private/InstancedFoliage.cpp
--- a/Engine/Source/Runtime/Foliage/Private/InstancedFoliage.cpp
+++ b/Engine/Source/Runtime/Foliage/Private/InstancedFoliage.cpp
@@ -4248,7 +4248,7 @@ void AInstancedFoliageActor::RemoveBaseComponentOnFoliageTypeInstances(UFoliageT
 	}
 }
 
-void AInstancedFoliageActor::AddInstances(UObject* WorldContextObject, UFoliageType* InFoliageType, TArray<FTransform> InTransforms)
+void AInstancedFoliageActor::AddInstances(UObject* WorldContextObject, UFoliageType* InFoliageType, const TArray<FTransform>& InTransforms)
 {
 	UWorld* World = GEngine->GetWorldFromContextObject(WorldContextObject, EGetWorldErrorMode::LogAndReturnNull);
 	if (World)
diff --git a/Engine/Source/Runtime/Foliage/Public/InstancedFoliageActor.h b/Engine/Source/Runtime/Foliage/Public/InstancedFoliageActor.h
--- a/Engine/Source/Runtime/Foliage/Public/InstancedFoliageActor.h
+++ b/Engine/Source/Runtime/Foliage/Public/InstancedFoliageActor.h
@@ -234,7 +234,7 @@ public:
 	void RemoveBaseComponentOnFoliageTypeInstances(UFoliageType* FoliageType);
 
 	UFUNCTION(BlueprintCallable, Category="Foliage", meta = (WorldContext = "WorldContextObject"))
-	static void AddInstances(UObject* WorldContextObject, UFoliageType* InFoliageType, TArray<FTransform> InTransforms);
+	static void AddInstances(UObject* WorldContextObject, UFoliageType* InFoliageType, const TArray<FTransform>& InTransforms);
 
 	UFUNCTION(BlueprintCallable, Category="Foliage", meta = (WorldContext = "WorldContextObject"))
 	static void RemoveAllInstances(UObject* WorldContextObject, UFoliageType* InFoliageType);
