Expose contact shadow flag for foliage

diff --git a/Engine/Source/Runtime/Foliage/Private/InstancedFoliage.cpp b/Engine/Source/Runtime/Foliage/Private/InstancedFoliage.cpp
--- a/Engine/Source/Runtime/Foliage/Private/InstancedFoliage.cpp
+++ b/Engine/Source/Runtime/Foliage/Private/InstancedFoliage.cpp
@@ -478,6 +478,7 @@ UFoliageType::UFoliageType(const FObjectInitializer& ObjectInitializer)
 	CastShadow = true;
 	bCastDynamicShadow = true;
 	bCastStaticShadow = true;
+	bCastContactShadow = true;
 	bAffectDynamicIndirectLighting = false;
 	// Most of the high instance count foliage like grass causes performance problems with distance field lighting
 	bAffectDistanceFieldLighting = false;
@@ -1421,6 +1422,12 @@ void FFoliageStaticMesh::UpdateComponentSettings(const UFoliageType_InstancedSta
 			bNeedsMarkRenderStateDirty = true;
 			bNeedsInvalidateLightingCache = true;
 		}
+		if (Component->bCastContactShadow != FoliageType->bCastContactShadow)
+		{
+			Component->bCastContactShadow = FoliageType->bCastContactShadow;
+			bNeedsMarkRenderStateDirty = true;
+			bNeedsInvalidateLightingCache = true;
+		}
 		if (Component->RuntimeVirtualTextures != FoliageType->RuntimeVirtualTextures)
 		{
 			Component->RuntimeVirtualTextures = FoliageType->RuntimeVirtualTextures;
diff --git a/Engine/Source/Runtime/Foliage/Public/FoliageType.h b/Engine/Source/Runtime/Foliage/Public/FoliageType.h
--- a/Engine/Source/Runtime/Foliage/Public/FoliageType.h
+++ b/Engine/Source/Runtime/Foliage/Public/FoliageType.h
@@ -307,6 +307,10 @@ public:
 	UPROPERTY(EditAnywhere, BlueprintReadOnly, Category=InstanceSettings, meta=(EditCondition="CastShadow"))
 	uint32 bCastStaticShadow:1;
 
+	/** Whether the object should cast contact shadows. This flag is only used if CastShadow is true. */
+	UPROPERTY(EditAnywhere, BlueprintReadOnly, Category=InstanceSettings, meta=(EditCondition = "CastShadow"))
+	uint8 bCastContactShadow : 1;
+
 	/** Whether this foliage should cast dynamic shadows as if it were a two sided material. */
 	UPROPERTY(EditAnywhere, AdvancedDisplay, BlueprintReadOnly, Category=InstanceSettings, meta=(EditCondition="bCastDynamicShadow"))
 	uint32 bCastShadowAsTwoSided:1;
