Expose FoliageType placement properties to BP

diff --git a/Engine/Source/Runtime/CoreUObject/Public/UObject/NoExportTypes.h b/Engine/Source/Runtime/CoreUObject/Public/UObject/NoExportTypes.h
--- a/Engine/Source/Runtime/CoreUObject/Public/UObject/NoExportTypes.h
+++ b/Engine/Source/Runtime/CoreUObject/Public/UObject/NoExportTypes.h
@@ -1380,15 +1380,15 @@ struct FFrameNumberRange
  * An interval of floats, defined by inclusive min and max values
  * @note This is a mirror of TInterval<float>, defined in Interval.h
  */
-USTRUCT(noexport)
+USTRUCT(noexport, BlueprintType)
 struct FFloatInterval
 {
 	/** Values must be >= Min */
-	UPROPERTY(EditAnywhere, Category=Interval)
+	UPROPERTY(EditAnywhere, BlueprintReadWrite, Category=Interval)
 	float Min;
 
 	/** Values must be <= Max */
-	UPROPERTY(EditAnywhere, Category=Interval)
+	UPROPERTY(EditAnywhere, BlueprintReadWrite, Category=Interval)
 	float Max;
 };
 
diff --git a/Engine/Source/Runtime/Foliage/Public/FoliageType.h b/Engine/Source/Runtime/Foliage/Public/FoliageType.h
--- a/Engine/Source/Runtime/Foliage/Public/FoliageType.h
+++ b/Engine/Source/Runtime/Foliage/Public/FoliageType.h
@@ -207,59 +207,59 @@ public:
 	// PLACEMENT
 
 	/** Specifies a range from minimum to maximum of the offset to apply to a foliage instance's Z location */
-	UPROPERTY(EditAnywhere, Category=Placement, meta=(DisplayName="Z Offset", ReapplyCondition="ReapplyZOffset"))
+	UPROPERTY(EditAnywhere, BlueprintReadWrite, Category=Placement, meta=(DisplayName="Z Offset", ReapplyCondition="ReapplyZOffset"))
 	FFloatInterval ZOffset;
 
 	/** Whether foliage instances should have their angle adjusted away from vertical to match the normal of the surface they're painted on 
 	 *  If AlignToNormal is enabled and RandomYaw is disabled, the instance will be rotated so that the +X axis points down-slope */
-	UPROPERTY(EditAnywhere, Category=Placement, meta=(ReapplyCondition="ReapplyAlignToNormal"))
+	UPROPERTY(EditAnywhere, BlueprintReadWrite, Category=Placement, meta=(ReapplyCondition="ReapplyAlignToNormal"))
 	uint32 AlignToNormal:1;
 
 	/** The maximum angle in degrees that foliage instances will be adjusted away from the vertical */
-	UPROPERTY(EditAnywhere, Category=Placement, meta=(UIMin = 0, ClampMin = 0, UIMax = 359, ClampMax = 359, HideBehind="AlignToNormal"))
+	UPROPERTY(EditAnywhere, BlueprintReadWrite, Category=Placement, meta=(UIMin = 0, ClampMin = 0, UIMax = 359, ClampMax = 359, HideBehind="AlignToNormal"))
 	float AlignMaxAngle;
 
 	/** If selected, foliage instances will have a random yaw rotation around their vertical axis applied */
-	UPROPERTY(EditAnywhere, Category=Placement, meta=(ReapplyCondition="ReapplyRandomYaw"))
+	UPROPERTY(EditAnywhere, BlueprintReadWrite, Category=Placement, meta=(ReapplyCondition="ReapplyRandomYaw"))
 	uint32 RandomYaw:1;
 
 	/** A random pitch adjustment can be applied to each instance, up to the specified angle in degrees, from the original vertical */
-	UPROPERTY(EditAnywhere, Category=Placement, meta=(UIMin = 0, ClampMin = 0, UIMax = 359, ClampMax = 359, ReapplyCondition="ReapplyRandomPitchAngle"))
+	UPROPERTY(EditAnywhere, BlueprintReadWrite, Category=Placement, meta=(UIMin = 0, ClampMin = 0, UIMax = 359, ClampMax = 359, ReapplyCondition="ReapplyRandomPitchAngle"))
 	float RandomPitchAngle;
 
 	/* Foliage instances will only be placed on surfaces sloping in the specified angle range from the horizontal */
-	UPROPERTY(EditAnywhere, Category=Placement, meta=(UIMin=0, ClampMin = 0, UIMax = 359, ClampMax = 359, ReapplyCondition="ReapplyGroundSlope"))
+	UPROPERTY(EditAnywhere, BlueprintReadWrite, Category=Placement, meta=(UIMin=0, ClampMin = 0, UIMax = 359, ClampMax = 359, ReapplyCondition="ReapplyGroundSlope"))
 	FFloatInterval GroundSlopeAngle;
 
 	/* The valid altitude range where foliage instances will be placed, specified using minimum and maximum world coordinate Z values */
-	UPROPERTY(EditAnywhere, Category=Placement, meta=(ReapplyCondition="ReapplyHeight"))
+	UPROPERTY(EditAnywhere, BlueprintReadWrite, Category=Placement, meta=(ReapplyCondition="ReapplyHeight"))
 	FFloatInterval Height;
 
 	/** If layer names are specified, painting on landscape will limit the foliage to areas of landscape with the specified layers painted */
-	UPROPERTY(EditAnywhere, AdvancedDisplay, Category=Placement, meta=(ReapplyCondition="ReapplyLandscapeLayers", DisplayName="Inclusion Landscape Layers"))
+	UPROPERTY(EditAnywhere, BlueprintReadWrite, AdvancedDisplay, Category=Placement, meta=(ReapplyCondition="ReapplyLandscapeLayers", DisplayName="Inclusion Landscape Layers"))
 	TArray<FName> LandscapeLayers;
 
 	/** Specifies the minimum value above which the landscape layer weight value must be, in order for foliage instances to be placed in a specific area */
-	UPROPERTY(EditAnywhere, AdvancedDisplay, Category=Placement, meta=(UIMin=0, ClampMin = 0, UIMax = 1, ClampMax = 1, HideBehind="LandscapeLayers", DisplayName="Minimum Inclusion Landscape Weight"))
+	UPROPERTY(EditAnywhere, BlueprintReadWrite, AdvancedDisplay, Category=Placement, meta=(UIMin=0, ClampMin = 0, UIMax = 1, ClampMax = 1, HideBehind="LandscapeLayers", DisplayName="Minimum Inclusion Landscape Weight"))
 	float MinimumLayerWeight;
 
 	/** If layer names are specified, painting on landscape will exclude the foliage to areas of landscape without the specified layers painted */
-	UPROPERTY(EditAnywhere, AdvancedDisplay, Category = Placement, meta = (ReapplyCondition = "ReapplyLandscapeLayers"))
+	UPROPERTY(EditAnywhere, BlueprintReadWrite, AdvancedDisplay, Category = Placement, meta = (ReapplyCondition = "ReapplyLandscapeLayers"))
 	TArray<FName> ExclusionLandscapeLayers;
 
 	/** Specifies the minimum value above which the landscape exclusion layer weight value must be, in order for foliage instances to be excluded in a specific area */
-	UPROPERTY(EditAnywhere, AdvancedDisplay, Category = Placement, meta = (UIMin=0, ClampMin = 0, UIMax = 1, ClampMax = 1, HideBehind = "ExclusionLandscapeLayers"))
+	UPROPERTY(EditAnywhere, BlueprintReadWrite, AdvancedDisplay, Category = Placement, meta = (UIMin=0, ClampMin = 0, UIMax = 1, ClampMax = 1, HideBehind = "ExclusionLandscapeLayers"))
 	float MinimumExclusionLayerWeight;
 
 	UPROPERTY()
 	FName LandscapeLayer_DEPRECATED;
 	
 	/* If checked, an overlap test with existing world geometry is performed before each instance is placed */
-	UPROPERTY(EditAnywhere, AdvancedDisplay, Category=Placement, meta=(ReapplyCondition="ReapplyCollisionWithWorld"))
+	UPROPERTY(EditAnywhere, BlueprintReadWrite, AdvancedDisplay, Category=Placement, meta=(ReapplyCondition="ReapplyCollisionWithWorld"))
 	uint32 CollisionWithWorld:1;
 
 	/* The foliage instance's collision bounding box will be scaled by the specified amount before performing the overlap check */
-	UPROPERTY(EditAnywhere, AdvancedDisplay, Category=Placement, meta=(HideBehind="CollisionWithWorld"))
+	UPROPERTY(EditAnywhere, BlueprintReadWrite, AdvancedDisplay, Category=Placement, meta=(HideBehind="CollisionWithWorld"))
 	FVector CollisionScale;
 
 	UPROPERTY()
diff --git a/Engine/Source/Runtime/Foliage/Public/FoliageType_InstancedStaticMesh.h b/Engine/Source/Runtime/Foliage/Public/FoliageType_InstancedStaticMesh.h
--- a/Engine/Source/Runtime/Foliage/Public/FoliageType_InstancedStaticMesh.h
+++ b/Engine/Source/Runtime/Foliage/Public/FoliageType_InstancedStaticMesh.h
@@ -14,10 +14,10 @@ class UFoliageType_InstancedStaticMesh : public UFoliageType
 {
 	GENERATED_UCLASS_BODY()
 
-	UPROPERTY(EditAnywhere, Category=Mesh, meta=(DisplayThumbnail="true"))
+	UPROPERTY(EditAnywhere, BlueprintReadWrite, Category=Mesh, meta=(DisplayThumbnail="true"))
 	UStaticMesh* Mesh;
 
-	UPROPERTY(EditAnywhere, AdvancedDisplay, Category = Mesh, Meta = (ToolTip = "Material overrides for foliage instances."))
+	UPROPERTY(EditAnywhere, BlueprintReadWrite, AdvancedDisplay, Category = Mesh, Meta = (ToolTip = "Material overrides for foliage instances."))
 	TArray<class UMaterialInterface*> OverrideMaterials;
 		
 	/** The component class to use for foliage instances. 
