Foliage: Cleanup (split from my actual new dev to avoid noise in review)
- FFoliageInfo now has an editor only IFA pointer
- FFoliageImpl now has an editor only FFoliageInfo pointer
- FoliageInfos is now private

diff --git a/Engine/Source/Editor/AssetPlacementEdMode/Private/AssetPlacementEdMode.cpp b/Engine/Source/Editor/AssetPlacementEdMode/Private/AssetPlacementEdMode.cpp
--- a/Engine/Source/Editor/AssetPlacementEdMode/Private/AssetPlacementEdMode.cpp
+++ b/Engine/Source/Editor/AssetPlacementEdMode/Private/AssetPlacementEdMode.cpp
@@ -211,11 +211,12 @@ void UAssetPlacementEdMode::DeleteSelection()
 	{
 		if (AInstancedFoliageActor* FoliageActor = *It)
 		{
-			for (auto& FoliageInfo : FoliageActor->FoliageInfos)
+			FoliageActor->ForEachFoliageInfo([](UFoliageType* FoliageType, FFoliageInfo& FoliageInfo)
 			{
-				TArray<int32> SelectedIndices = FoliageInfo.Value->SelectedIndices.Array();
-				FoliageInfo.Value->RemoveInstances(FoliageActor, SelectedIndices, true);
-			}
+				TArray<int32> SelectedIndices = FoliageInfo.SelectedIndices.Array();
+				FoliageInfo.RemoveInstances(SelectedIndices, true);
+				return true; // continue iteration
+			});
 		}
 	}
 
@@ -235,10 +236,11 @@ void UAssetPlacementEdMode::ClearSelection()
 	{
 		if (AInstancedFoliageActor* FoliageActor = *It)
 		{
-			for (auto& FoliageInfo : FoliageActor->FoliageInfos)
+			FoliageActor->ForEachFoliageInfo([](UFoliageType* FoliageType, FFoliageInfo& FoliageInfo)
 			{
-				FoliageInfo.Value->ClearSelection();
-			}
+				FoliageInfo.ClearSelection();
+				return true; // continue iteration
+			});
 		}
 	}
 	GetToolManager()->EndUndoTransaction();
diff --git a/Engine/Source/Editor/AssetPlacementEdMode/Private/Tools/PlacementBrushToolBase.cpp b/Engine/Source/Editor/AssetPlacementEdMode/Private/Tools/PlacementBrushToolBase.cpp
--- a/Engine/Source/Editor/AssetPlacementEdMode/Private/Tools/PlacementBrushToolBase.cpp
+++ b/Engine/Source/Editor/AssetPlacementEdMode/Private/Tools/PlacementBrushToolBase.cpp
@@ -194,7 +194,7 @@ TArray<FTypedElementHandle> UPlacementBrushToolBase::GetElementsInBrushRadius()
 		FActorPartitionGetParams PartitionActorFindParams(AInstancedFoliageActor::StaticClass(), bCreatePartitionActorIfMissing, GEditor->GetEditorWorldContext().World()->GetCurrentLevel(), LastBrushStamp.WorldPosition);
 		if (AInstancedFoliageActor* FoliageActor = Cast<AInstancedFoliageActor>(PartitionSubsystem->GetActor(PartitionActorFindParams)))
 		{
-			for (auto& FoliageInfo : FoliageActor->FoliageInfos)
+			for (const auto& FoliageInfo : FoliageActor->GetFoliageInfos())
 			{
 				FTypedElementHandle SourceObjectHandle = UEngineElementsLibrary::AcquireEditorObjectElementHandle(FoliageInfo.Key->GetSource());
 				if (UAssetPlacementEdMode::DoesPaletteSupportElement(SourceObjectHandle, PlacementSettings->PaletteItems))
diff --git a/Engine/Source/Editor/AssetPlacementEdMode/Private/Tools/PlacementEraseTool.cpp b/Engine/Source/Editor/AssetPlacementEdMode/Private/Tools/PlacementEraseTool.cpp
--- a/Engine/Source/Editor/AssetPlacementEdMode/Private/Tools/PlacementEraseTool.cpp
+++ b/Engine/Source/Editor/AssetPlacementEdMode/Private/Tools/PlacementEraseTool.cpp
@@ -57,17 +57,18 @@ void UPlacementModeEraseTool::OnTick(float DeltaTime)
 			// Since the foliage static mesh instances do not currently operate with element handles, we have to drill in manually here.
 			if (AInstancedFoliageActor* FoliageActor = Cast<AInstancedFoliageActor>(Actor))
 			{
-				for (auto& FoliageInfo : FoliageActor->FoliageInfos)
+				FoliageActor->ForEachFoliageInfo([this](UFoliageType* FoliageType, FFoliageInfo& FoliageInfo)
 				{
-					FTypedElementHandle SourceObjectHandle = UEngineElementsLibrary::AcquireEditorObjectElementHandle(FoliageInfo.Key->GetSource());
+					FTypedElementHandle SourceObjectHandle = UEngineElementsLibrary::AcquireEditorObjectElementHandle(FoliageType->GetSource());
 					if (UAssetPlacementEdMode::DoesPaletteSupportElement(SourceObjectHandle, PlacementSettings->PaletteItems))
 					{
 						TArray<int32> Instances;
 						FSphere SphereToCheck(LastBrushStamp.WorldPosition, LastBrushStamp.Radius);
-						FoliageInfo.Value->GetInstancesInsideSphere(SphereToCheck, Instances);
-						FoliageInfo.Value->RemoveInstances(FoliageActor, Instances, true);
+						FoliageInfo.GetInstancesInsideSphere(SphereToCheck, Instances);
+						FoliageInfo.RemoveInstances(Instances, true);
 					}
-				}
+					return true; // continue iteraton
+				});
 			}
 			else if (ActorSubsystem)
 			{
diff --git a/Engine/Source/Editor/FoliageEdit/Private/FoliageEdMode.cpp b/Engine/Source/Editor/FoliageEdit/Private/FoliageEdMode.cpp
--- a/Engine/Source/Editor/FoliageEdit/Private/FoliageEdMode.cpp
+++ b/Engine/Source/Editor/FoliageEdit/Private/FoliageEdMode.cpp
@@ -733,12 +733,12 @@ void FEdModeFoliage::OnObjectsReplaced(const TMap<UObject*, UObject*>& Replaceme
 			if (UFoliageType* ReplacedFoliageType = Cast<UFoliageType>(ReplacementPair.Key))
 			{
 				TUniqueObj<FFoliageInfo> FoliageInfo;
-				if (IFA->FoliageInfos.RemoveAndCopyValue(ReplacedFoliageType, FoliageInfo))
+				if (IFA->RemoveFoliageInfoAndCopyValue(ReplacedFoliageType, FoliageInfo))
 				{
 					// Re-add the unique mesh info associated with the replaced foliage type
 					UFoliageType* ReplacementFoliageType = Cast<UFoliageType>(ReplacementPair.Value);
-					TUniqueObj<FFoliageInfo>& NewFoliageInfo = IFA->FoliageInfos.Add(ReplacementFoliageType, MoveTemp(FoliageInfo));
-					NewFoliageInfo->ReallocateClusters(IFA, ReplacementFoliageType);
+					TUniqueObj<FFoliageInfo>& NewFoliageInfo = IFA->AddFoliageInfo(ReplacementFoliageType, MoveTemp(FoliageInfo));
+					NewFoliageInfo->ReallocateClusters(ReplacementFoliageType);
 
 					bAnyFoliageTypeReplaced = true;
 				}
@@ -1375,10 +1375,10 @@ static void SpawnFoliageInstance(UWorld* InWorld, const UFoliageType* Settings,
 		FFoliageInfo* Info = nullptr;
 		UFoliageType* FoliageSettings = IFA->AddFoliageType(Settings, &Info);
 
-		Info->AddInstances(IFA, FoliageSettings, PlacedLevelInstances.Value);
+		Info->AddInstances(FoliageSettings, PlacedLevelInstances.Value);
 		if (InRebuildFoliageTree)
 		{
-			Info->Refresh(IFA, true, false);
+			Info->Refresh(true, false);
 		}
 	}
 }
@@ -1393,7 +1393,7 @@ void FEdModeFoliage::RebuildFoliageTree(const UFoliageType* Settings)
 
 			if (FoliageInfo != nullptr)
 			{
-				FoliageInfo->Refresh(IFA, true, false);
+				FoliageInfo->Refresh(true, false);
 			}
 		}
 	}
@@ -1629,7 +1629,7 @@ void FEdModeFoliage::RemoveInstancesForBrush(UWorld* InWorld, const UFoliageType
 		{
 			CurrentFoliageTraceBrushAffectedIFAs.Add(IFA);
 
-			FoliageInfo->RemoveInstances(IFA, PotentialInstancesToRemove, false);
+			FoliageInfo->RemoveInstances(PotentialInstancesToRemove, false);
 		}
 		return true;
 	};
@@ -1652,7 +1652,7 @@ void FEdModeFoliage::SelectInstanceAtLocation(UWorld* InWorld, const UFoliageTyp
 		{
 			TArray<int32> Instances;
 			Instances.Add(Instance);
-			FoliageInfo->SelectInstances(IFA, bSelect, Instances);
+			FoliageInfo->SelectInstances(bSelect, Instances);
 		}
 
 		return true;
@@ -1673,7 +1673,7 @@ void FEdModeFoliage::SelectInstancesForBrush(UWorld* InWorld, const UFoliageType
 		{
 			return true;
 		}
-		FoliageInfo->SelectInstances(IFA, bSelect, Instances);
+		FoliageInfo->SelectInstances(bSelect, Instances);
 
 		return true;
 	};
@@ -1717,10 +1717,10 @@ void FEdModeFoliage::ExcludeFoliageActors(const TArray<const UFoliageType *>& Fo
 		{
 			for (auto& Pair : ActorFoliageTypes)
 			{
-				if (TUniqueObj<FFoliageInfo>* FoliageInfoPtr = IFA->FoliageInfos.Find(Pair.Value))
+				if (FFoliageInfo* FoliageInfoPtr = IFA->FindInfo(Pair.Value))
 				{
 					IFA->Modify();
-					(*FoliageInfoPtr)->ExcludeActors();
+					FoliageInfoPtr->ExcludeActors();
 					OnInstanceCountUpdated(Pair.Value);
 				}
 			}
@@ -1776,7 +1776,7 @@ void FEdModeFoliage::IncludeNonFoliageActors(const TArray<const UFoliageType*>&
 							if (FoliageInfo)
 							{
 
-								FoliageInfo->IncludeActor(IFA, FoliageType, CurrentActor);
+								FoliageInfo->IncludeActor(FoliageType, CurrentActor);
 								UpdateFoliageTypes.Add(FoliageType);
 							}
 						}
@@ -1844,7 +1844,7 @@ void FEdModeFoliage::SelectInstances(UWorld* InWorld, const UFoliageType* Settin
 		FFoliageInfo* FoliageInfo = (*It);
 		AInstancedFoliageActor* IFA = It.GetActor();
 
-		FoliageInfo->SelectInstances(IFA, bSelect);
+		FoliageInfo->SelectInstances(bSelect);
 	}
 }
 
@@ -1871,13 +1871,12 @@ void FEdModeFoliage::UpdateInstancePartitioning(UWorld* InWorld)
 	for (TActorIterator<AInstancedFoliageActor> It(InWorld); It; ++It)
 	{
 		AInstancedFoliageActor* IFA = *It;
-		for (auto& MeshPair : IFA->FoliageInfos)
+		IFA->ForEachFoliageInfo([InWorld, IFA, ActorPartitionSubsystem](UFoliageType* FoliageType, FFoliageInfo& FoliageInfo)
 		{
-			FFoliageInfo& FoliageInfo = *MeshPair.Value;
 			if (FoliageInfo.Type == EFoliageImplType::Actor && 
 				ActorPartitionSubsystem->IsLevelPartition())
 			{
-				continue; // Actors are handled through the Partitioning code
+				return true; // Actors are handled through the Partitioning code
 			}
 
 			AInstancedFoliageActor* TargetIFA = nullptr;
@@ -1903,12 +1902,14 @@ void FEdModeFoliage::UpdateInstancePartitioning(UWorld* InWorld)
 				if (InstancesToMove.Num())
 				{
 					// TargetIFA can be null (if target is unloaded cell). Meaning instances will be deleted.
-					FoliageInfo.MoveInstances(IFA, TargetIFA, InstancesToMove, true);
+					FoliageInfo.MoveInstances(TargetIFA, InstancesToMove, true);
 					bMovedInstances = true;
 				}
 
 			} while (bMovedInstances && FoliageInfo.SelectedIndices.Num() > 0);
-		}
+
+			return true; // continue iteration
+		});
 	}
 }
 
@@ -1925,17 +1926,17 @@ void FEdModeFoliage::PostTransformSelectedInstances(UWorld* InWorld)
 		AInstancedFoliageActor* IFA = *It;
 		bool bFoundSelection = false;
 
-		for (auto& MeshPair : IFA->FoliageInfos)
+		IFA->ForEachFoliageInfo([](UFoliageType* FoliageType, FFoliageInfo& FoliageInfo)
 		{
-			FFoliageInfo& FoliageInfo = *MeshPair.Value;
 			TArray<int32> SelectedIndices = FoliageInfo.SelectedIndices.Array();
 
 			if (SelectedIndices.Num() > 0)
 			{
 				const bool bFinished = true;
-				FoliageInfo.PostMoveInstances(IFA, SelectedIndices, bFinished);
+				FoliageInfo.PostMoveInstances(SelectedIndices, bFinished);
 			}
-		}
+			return true; // continue iteration
+		});
 	}
 }
 
@@ -1946,9 +1947,8 @@ void FEdModeFoliage::TransformSelectedInstances(UWorld* InWorld, const FVector&
 		AInstancedFoliageActor* IFA = *It;
 		bool bFoundSelection = false;
 
-		for (auto& MeshPair : IFA->FoliageInfos)
+		IFA->ForEachFoliageInfo([this, IFA, &bFoundSelection, InDrag, InRot, InScale, bDuplicate](UFoliageType* FoliageType, FFoliageInfo& FoliageInfo)
 		{
-			FFoliageInfo& FoliageInfo = *MeshPair.Value;
 			TArray<int32> SelectedIndices = FoliageInfo.SelectedIndices.Array();
 
 			if (SelectedIndices.Num() > 0)
@@ -1962,12 +1962,12 @@ void FEdModeFoliage::TransformSelectedInstances(UWorld* InWorld, const FVector&
 
 				if (bDuplicate)
 				{
-					FoliageInfo.DuplicateInstances(IFA, MeshPair.Key, SelectedIndices);
-					OnInstanceCountUpdated(MeshPair.Key);
+					FoliageInfo.DuplicateInstances(FoliageType, SelectedIndices);
+					OnInstanceCountUpdated(FoliageType);
 				}
 
 				bMoving = true;
-				FoliageInfo.PreMoveInstances(IFA, SelectedIndices);
+				FoliageInfo.PreMoveInstances(SelectedIndices);
 
 				for (int32 SelectedInstanceIdx : SelectedIndices)
 				{
@@ -1979,9 +1979,10 @@ void FEdModeFoliage::TransformSelectedInstances(UWorld* InWorld, const FVector&
 				}
 
 				const bool bFinished = false;
-				FoliageInfo.PostMoveInstances(IFA, SelectedIndices, bFinished);
+				FoliageInfo.PostMoveInstances(SelectedIndices, bFinished);
 			}
-		}
+			return true; // continue iteration
+		});
 
 		if (bFoundSelection)
 		{
@@ -2026,7 +2027,7 @@ void FEdModeFoliage::RemoveSelectedInstances(UWorld* InWorld)
 	{
 		AInstancedFoliageActor* IFA = *It;
 		bool bHasSelection = false;
-		for (auto& MeshPair : IFA->FoliageInfos)
+		for (auto& MeshPair : IFA->GetFoliageInfos())
 		{
 			if (MeshPair.Value->SelectedIndices.Num())
 			{
@@ -2038,17 +2039,17 @@ void FEdModeFoliage::RemoveSelectedInstances(UWorld* InWorld)
 		if (bHasSelection)
 		{
 			IFA->Modify();
-			for (auto& MeshPair : IFA->FoliageInfos)
+			IFA->ForEachFoliageInfo([this](UFoliageType* FoliageType, FFoliageInfo& FoliageInfo)
 			{
-				FFoliageInfo& Mesh = *MeshPair.Value;
-				if (Mesh.SelectedIndices.Num() > 0)
+				if (FoliageInfo.SelectedIndices.Num() > 0)
 				{
-					TArray<int32> InstancesToDelete = Mesh.SelectedIndices.Array();
-					Mesh.RemoveInstances(IFA, InstancesToDelete, true);
+					TArray<int32> InstancesToDelete = FoliageInfo.SelectedIndices.Array();
+					FoliageInfo.RemoveInstances(InstancesToDelete, true);
 
-					OnInstanceCountUpdated(MeshPair.Key);
+					OnInstanceCountUpdated(FoliageType);
 				}
-			}
+				return true; // continue iteration
+			});
 		}
 	}
 }
@@ -2070,7 +2071,7 @@ void FEdModeFoliage::GetSelectedInstanceFoliageTypes(TArray<const UFoliageType*>
 	{
 		AInstancedFoliageActor* IFA = *It;
 		bool bHasSelection = false;
-		for (auto& MeshPair : IFA->FoliageInfos)
+		for (auto& MeshPair : IFA->GetFoliageInfos())
 		{
 			if (MeshPair.Value->SelectedIndices.Num())
 			{
@@ -2171,7 +2172,7 @@ void FEdModeFoliage::SelectInvalidInstances(const UFoliageType* Settings)
 
 		if (InvalidInstances.Num() > 0)
 		{
-			FoliageInfo->SelectInstances(IFA, true, InvalidInstances);
+			FoliageInfo->SelectInstances(true, InvalidInstances);
 		}
 	}
 }
@@ -2528,13 +2529,13 @@ void FEdModeFoliage::ReapplyInstancesForBrush(UWorld* InWorld, AInstancedFoliage
 
 	if (UpdatedInstances.Num() > 0)
 	{
-		FoliageInfo->PostUpdateInstances(IFA, UpdatedInstances);
+		FoliageInfo->PostUpdateInstances(UpdatedInstances);
 		IFA->RegisterAllComponents();
 	}
 
 	if (InstancesToDelete.Num())
 	{
-		FoliageInfo->RemoveInstances(IFA, InstancesToDelete.Array(), true);
+		FoliageInfo->RemoveInstances(InstancesToDelete.Array(), true);
 	}
 }
 
@@ -3031,9 +3032,8 @@ void FEdModeFoliage::SnapSelectedInstancesToGround(UWorld* InWorld)
 			AInstancedFoliageActor* IFA = *It;
 			bool bFoundSelection = false;
 
-			for (auto& MeshPair : IFA->FoliageInfos)
+			IFA->ForEachFoliageInfo([this, IFA, &bFoundSelection, &bMovedInstance](UFoliageType* FoliageType, FFoliageInfo& FoliageInfo)
 			{
-				FFoliageInfo& FoliageInfo = *MeshPair.Value;
 				TArray<int32> SelectedIndices = FoliageInfo.SelectedIndices.Array();
 
 				if (SelectedIndices.Num() > 0)
@@ -3045,16 +3045,17 @@ void FEdModeFoliage::SnapSelectedInstancesToGround(UWorld* InWorld)
 						bFoundSelection = true;
 					}
 
-					FoliageInfo.PreMoveInstances(IFA, SelectedIndices);
+					FoliageInfo.PreMoveInstances(SelectedIndices);
 
 					for (int32 InstanceIndex : SelectedIndices)
 					{
-						bMovedInstance |= SnapInstanceToGround(IFA, MeshPair.Key->AlignMaxAngle, FoliageInfo, InstanceIndex);
+						bMovedInstance |= SnapInstanceToGround(IFA, FoliageType->AlignMaxAngle, FoliageInfo, InstanceIndex);
 					}
 
-					FoliageInfo.PostMoveInstances(IFA, SelectedIndices);
+					FoliageInfo.PostMoveInstances(SelectedIndices);
 				}
-			}
+				return true; // continue iteration
+			});
 		}
 
 		if (bMovedInstance)
@@ -3161,7 +3162,7 @@ void FEdModeFoliage::PopulateFoliageMeshList()
 	for (TActorIterator<AInstancedFoliageActor> It(World); It; ++It)
 	{
 		AInstancedFoliageActor* IFA = *It;
-		for (auto& MeshPair : IFA->FoliageInfos)
+		for (const auto& MeshPair : IFA->GetFoliageInfos())
 		{
 			//@todo_ow: this doesn't make sense in WP
 			if (!CanPaint(MeshPair.Key, CurrentLevel))
@@ -3258,7 +3259,7 @@ void FEdModeFoliage::CalcTotalInstanceCount(int32& OutInstanceCountTotal, int32&
 	{
 		AInstancedFoliageActor* IFA = *It;
 		int32 IFAInstanceCount = 0;
-		for (const auto& MeshPair : IFA->FoliageInfos)
+		for (const auto& MeshPair : IFA->GetFoliageInfos())
 		{
 			const FFoliageInfo& FoliageInfo = *MeshPair.Value;
 			IFAInstanceCount += FoliageInfo.Instances.Num();
@@ -3386,7 +3387,7 @@ void FEdModeFoliage::MoveSelectedFoliageToLevel(ULevel* InTargetLevel)
 						{
 							// Restore previous selection for move operation
 							FFoliageInfo* FoliageInfo = IFA->FindInfo(NewFoliageType);
-							FoliageInfo->SelectInstances(IFA, true, PreviousSelectionArray);
+							FoliageInfo->SelectInstances(true, PreviousSelectionArray);
 						}
 					}
 				}
@@ -3529,7 +3530,7 @@ void FEdModeFoliage::BakeFoliage(UFoliageType* Settings, bool bSelectedOnly)
 		}
 
 		// Remove
-		FoliageInfo->RemoveInstances(IFA, InstancesToConvert, true);
+		FoliageInfo->RemoveInstances(InstancesToConvert, true);
 	}
 }
 
@@ -3542,10 +3543,10 @@ UFoliageType* FEdModeFoliage::CopySettingsObject(UFoliageType* Settings)
 	IFA->Modify();
 
 	TUniqueObj<FFoliageInfo> FoliageInfo;
-	if (IFA->FoliageInfos.RemoveAndCopyValue(Settings, FoliageInfo))
+	if (IFA->RemoveFoliageInfoAndCopyValue(Settings, FoliageInfo))
 	{
 		Settings = (UFoliageType*)StaticDuplicateObject(Settings, IFA, NAME_None, RF_AllFlags & ~(RF_Standalone | RF_Public));
-		IFA->FoliageInfos.Add(Settings, MoveTemp(FoliageInfo));
+		IFA->AddFoliageInfo(Settings, MoveTemp(FoliageInfo));
 		return Settings;
 	}
 	else
@@ -3583,7 +3584,7 @@ void FEdModeFoliage::ReallocateClusters(UFoliageType* Settings)
 	for (FFoliageInfoIterator It(World, Settings); It; ++It)
 	{
 		FFoliageInfo* FoliageInfo = (*It);
-		FoliageInfo->ReallocateClusters(It.GetActor(), Settings);
+		FoliageInfo->ReallocateClusters(Settings);
 	}
 }
 
diff --git a/Engine/Source/Editor/FoliageEdit/Private/FoliageEditUtility.cpp b/Engine/Source/Editor/FoliageEdit/Private/FoliageEditUtility.cpp
--- a/Engine/Source/Editor/FoliageEdit/Private/FoliageEditUtility.cpp
+++ b/Engine/Source/Editor/FoliageEdit/Private/FoliageEditUtility.cpp
@@ -86,7 +86,7 @@ void FFoliageEditUtility::ReplaceFoliageTypeObject(UWorld* InWorld, UFoliageType
 		AInstancedFoliageActor* IFA = *It;
 		IFA->Modify();
 		TUniqueObj<FFoliageInfo> OldInfo;
-		IFA->FoliageInfos.RemoveAndCopyValue(OldType, OldInfo);
+		IFA->RemoveFoliageInfoAndCopyValue(OldType, OldInfo);
 
 		// Old component needs to go
 		if (OldInfo->IsInitialized())
@@ -96,17 +96,17 @@ void FFoliageEditUtility::ReplaceFoliageTypeObject(UWorld* InWorld, UFoliageType
 
 		// Append instances if new foliage type is already exists in this actor
 		// Otherwise just replace key entry for instances
-		TUniqueObj<FFoliageInfo>* NewInfo = IFA->FoliageInfos.Find(NewType);
+		FFoliageInfo* NewInfo = IFA->FindInfo(NewType);
 		if (NewInfo)
 		{
-			(*NewInfo)->Instances.Append(OldInfo->Instances);
-			(*NewInfo)->ReallocateClusters(IFA, NewType);
+			NewInfo->Instances.Append(OldInfo->Instances);
+			NewInfo->ReallocateClusters(NewType);
 		}
 		else
 		{
 			// Make sure if type changes we have proper implementation
-			TUniqueObj<FFoliageInfo>& NewFoliageInfo = IFA->FoliageInfos.Add(NewType, MoveTemp(OldInfo));
-			NewFoliageInfo->ReallocateClusters(IFA, NewType);
+			TUniqueObj<FFoliageInfo>& NewFoliageInfo = IFA->AddFoliageInfo(NewType, MoveTemp(OldInfo));
+			NewFoliageInfo->ReallocateClusters(NewType);
 		}
 	}
 }
@@ -179,7 +179,7 @@ void FFoliageEditUtility::MoveActorFoliageInstancesToLevel(ULevel* InTargetLevel
 					{
 						// Restore previous selection for move operation
 						FFoliageInfo* MeshInfo = IFA->FindInfo(NewFoliageType);
-						MeshInfo->SelectInstances(IFA, true, PreviousSelectionArray);
+						MeshInfo->SelectInstances(true, PreviousSelectionArray);
 					}
 				}
 			}
diff --git a/Engine/Source/Editor/UnrealEd/Private/Commandlets/WorldPartitionConvertCommandlet.cpp b/Engine/Source/Editor/UnrealEd/Private/Commandlets/WorldPartitionConvertCommandlet.cpp
--- a/Engine/Source/Editor/UnrealEd/Private/Commandlets/WorldPartitionConvertCommandlet.cpp
+++ b/Engine/Source/Editor/UnrealEd/Private/Commandlets/WorldPartitionConvertCommandlet.cpp
@@ -861,36 +861,32 @@ int32 UWorldPartitionConvertCommandlet::Main(const FString& Params)
 	{
 		TRACE_CPUPROFILER_EVENT_SCOPE(PartitionFoliage);
 
-		for (auto& Pair : IFA->FoliageInfos)
+		IFA->ForEachFoliageInfo([IFA, ActorPartitionSubsystem, MainWorld](UFoliageType* FoliageType, FFoliageInfo& FoliageInfo)
 		{
-			FFoliageInfo* FoliageInfo = &Pair.Value.Get();
-			if (FoliageInfo)
+			if (FoliageInfo.Type == EFoliageImplType::Actor)
 			{
-				if (FoliageInfo->Type == EFoliageImplType::Actor)
-				{
-					// We don't support Actor Foliage in WP
-					FoliageInfo->ExcludeActors();
-				}
-				else if (FoliageInfo->Instances.Num())
+				// We don't support Actor Foliage in WP
+				FoliageInfo.ExcludeActors();
+			}
+			else if (FoliageInfo.Instances.Num())
+			{
+				FBox InstanceBounds = FoliageInfo.GetApproximatedInstanceBounds();
+				FActorPartitionGridHelper::ForEachIntersectingCell(AInstancedFoliageActor::StaticClass(), InstanceBounds, MainWorld->PersistentLevel, [IFA, &FoliageInfo, ActorPartitionSubsystem](const UActorPartitionSubsystem::FCellCoord& InCellCoord, const FBox& InCellBounds)
 				{
-					FBox InstanceBounds = FoliageInfo->GetApproximatedInstanceBounds();
-					FActorPartitionGridHelper::ForEachIntersectingCell(AInstancedFoliageActor::StaticClass(), InstanceBounds, MainWorld->PersistentLevel, [IFA, FoliageInfo, ActorPartitionSubsystem](const UActorPartitionSubsystem::FCellCoord& InCellCoord, const FBox& InCellBounds)
+					TArray<int32> Indices;
+					FoliageInfo.GetInstancesInsideBounds(InCellBounds, Indices);
+					if (Indices.Num())
 					{
-						TArray<int32> Indices;
-						FoliageInfo->GetInstancesInsideBounds(InCellBounds, Indices);
-						if (Indices.Num())
-						{
-							AInstancedFoliageActor* CellIFA = Cast<AInstancedFoliageActor>(ActorPartitionSubsystem->GetActor(AInstancedFoliageActor::StaticClass(), InCellCoord, true));
-							check(CellIFA);
-							FoliageInfo->MoveInstances(IFA, CellIFA, TSet<int32>(Indices), false);
-						}
-
-						return true;
-					});
-					check(!FoliageInfo->Instances.Num());
-				}
+						AInstancedFoliageActor* CellIFA = Cast<AInstancedFoliageActor>(ActorPartitionSubsystem->GetActor(AInstancedFoliageActor::StaticClass(), InCellCoord, true));
+						check(CellIFA);
+						FoliageInfo.MoveInstances(CellIFA, TSet<int32>(Indices), false);
+					}
+					return true;
+				});
+				check(!FoliageInfo.Instances.Num());
 			}
-		}
+			return true; // continue iterating
+		});
 
 		IFA->GetLevel()->GetWorld()->DestroyActor(IFA);
 	};
diff --git a/Engine/Source/Editor/UnrealEd/Private/EditorObject.cpp b/Engine/Source/Editor/UnrealEd/Private/EditorObject.cpp
--- a/Engine/Source/Editor/UnrealEd/Private/EditorObject.cpp
+++ b/Engine/Source/Editor/UnrealEd/Private/EditorObject.cpp
@@ -385,7 +385,7 @@ static const TCHAR* ImportProperties(
 						Algo::Transform(Pair.Value, InstancePtrs, [](const FFoliageInstance& FoliageInstance) { return &FoliageInstance; });
 						if (MeshInfo)
 						{
-							MeshInfo->AddInstances(IFA, FoliageType, InstancePtrs);
+							MeshInfo->AddInstances(FoliageType, InstancePtrs);
 						}
 					}
 				}
diff --git a/Engine/Source/Editor/UnrealEd/Private/EditorServer.cpp b/Engine/Source/Editor/UnrealEd/Private/EditorServer.cpp
--- a/Engine/Source/Editor/UnrealEd/Private/EditorServer.cpp
+++ b/Engine/Source/Editor/UnrealEd/Private/EditorServer.cpp
@@ -3268,7 +3268,7 @@ TArray<UFoliageType*> UEditorEngine::GetFoliageTypesInWorld(UWorld* InWorld)
 	// Iterate over all foliage actors in the world
 	for (TActorIterator<AInstancedFoliageActor> It(InWorld); It; ++It)
 	{
-		for (const auto& Pair : It->FoliageInfos)
+		for (const auto& Pair : It->GetFoliageInfos())
 		{
 			FoliageSet.Add(Pair.Key);
 		}
diff --git a/Engine/Source/Editor/UnrealEd/Private/Factories/EditorStaticMeshFactory.cpp b/Engine/Source/Editor/UnrealEd/Private/Factories/EditorStaticMeshFactory.cpp
--- a/Engine/Source/Editor/UnrealEd/Private/Factories/EditorStaticMeshFactory.cpp
+++ b/Engine/Source/Editor/UnrealEd/Private/Factories/EditorStaticMeshFactory.cpp
@@ -37,7 +37,7 @@ TArray<FTypedElementHandle> UEditorStaticMeshFactory::PlaceAsset(const FAssetPla
 			// Create or find the instanced static mesh foliage type for the given mesh
 			UObject* AssetDataObject = InPlacementInfo.AssetToPlace.GetAsset();
 			UFoliageType_InstancedStaticMesh* FoundOrCreatedType = nullptr;
-			for (auto& FoliageTypePair : FoliageActor->FoliageInfos)
+			for (const auto& FoliageTypePair : FoliageActor->GetFoliageInfos())
 			{
 				FAssetData FoliageTypeSource(FoliageTypePair.Key->GetSource());
 				if (FoliageTypeSource == InPlacementInfo.AssetToPlace)
@@ -59,7 +59,7 @@ TArray<FTypedElementHandle> UEditorStaticMeshFactory::PlaceAsset(const FAssetPla
 			FoliagePlacementInfo.Location = InPlacementInfo.FinalizedTransform.GetLocation();
 			FoliagePlacementInfo.Rotation = InPlacementInfo.FinalizedTransform.Rotator();
 			FoliagePlacementInfo.DrawScale3D = InPlacementInfo.FinalizedTransform.GetScale3D();
-			FoliageInfo->AddInstance(FoliageActor, FoundOrCreatedType, FoliagePlacementInfo);
+			FoliageInfo->AddInstance(FoundOrCreatedType, FoliagePlacementInfo);
 
 			// Todo: deal with returning instanced handles
 			// In the meantime, we return the handle of the the HISM component added.
diff --git a/Engine/Source/Runtime/Foliage/Private/FoliageActor.cpp b/Engine/Source/Runtime/Foliage/Private/FoliageActor.cpp
--- a/Engine/Source/Runtime/Foliage/Private/FoliageActor.cpp
+++ b/Engine/Source/Runtime/Foliage/Private/FoliageActor.cpp
@@ -60,7 +60,7 @@ bool FFoliageActor::IsInitialized() const
 	return ActorClass != nullptr;
 }
 
-void FFoliageActor::Initialize(AInstancedFoliageActor* IFA, const UFoliageType* FoliageType)
+void FFoliageActor::Initialize(const UFoliageType* FoliageType)
 {
 	check(!IsInitialized());
 	const UFoliageType_Actor* FoliageType_Actor = Cast<UFoliageType_Actor>(FoliageType);
@@ -75,12 +75,15 @@ void FFoliageActor::Uninitialize()
 	ActorClass = nullptr;
 }
 
-AActor* FFoliageActor::Spawn(AInstancedFoliageActor* IFA, const FFoliageInstance& Instance)
+AActor* FFoliageActor::Spawn(const FFoliageInstance& Instance)
 {
 	if (ActorClass == nullptr)
 	{
 		return nullptr;
 	}
+
+	AInstancedFoliageActor* IFA = GetIFA();
+
 	FEditorScriptExecutionGuard ScriptGuard;
 	FActorSpawnParameters SpawnParameters;
 	SpawnParameters.ObjectFlags = RF_Transactional;
@@ -117,28 +120,28 @@ int32 FFoliageActor::GetInstanceCount() const
 	return ActorInstances.Num();
 }
 
-void FFoliageActor::PreAddInstances(AInstancedFoliageActor* IFA, const UFoliageType* FoliageType, int32 Count)
+void FFoliageActor::PreAddInstances(const UFoliageType* FoliageType, int32 Count)
 {
 	if (!IsInitialized())
 	{
-		Initialize(IFA, FoliageType);
+		Initialize(FoliageType);
 		check(IsInitialized());
 	}
 }
 
-void FFoliageActor::AddInstance(AInstancedFoliageActor* IFA, const FFoliageInstance& NewInstance)
+void FFoliageActor::AddInstance(const FFoliageInstance& NewInstance)
 {
-	ActorInstances.Add(Spawn(IFA, NewInstance));
+	ActorInstances.Add(Spawn(NewInstance));
 }
 
-void FFoliageActor::AddExistingInstance(AInstancedFoliageActor* IFA, const FFoliageInstance& ExistingInstance, UObject* InstanceImplementation)
+void FFoliageActor::AddExistingInstance(const FFoliageInstance& ExistingInstance, UObject* InstanceImplementation)
 {
 	AActor* Actor = Cast<AActor>(InstanceImplementation);
 	check(Actor);
 	check(Actor->GetClass() == ActorClass);
 	Actor->SetActorTransform(ExistingInstance.GetInstanceWorldTransform());
 	FFoliageHelper::SetIsOwnedByFoliage(Actor);
-	check(IFA->GetLevel() == Actor->GetLevel());
+	check(GetIFA()->GetLevel() == Actor->GetLevel());
 	ActorInstances.Add(Actor);
 }
 
@@ -209,13 +212,13 @@ int32 FFoliageActor::FindIndex(const UPrimitiveComponent* Component) const
 	return FindIndex(Component->GetOwner());
 }
 
-void FFoliageActor::Refresh(AInstancedFoliageActor* IFA, const TArray<FFoliageInstance>& Instances, bool Async, bool Force)
+void FFoliageActor::Refresh(const TArray<FFoliageInstance>& Instances, bool Async, bool Force)
 {
 	for (int32 i = 0; i < Instances.Num(); ++i)
 	{
 		if (ActorInstances[i] == nullptr || ActorInstances[i]->IsPendingKill())
 		{
-			ActorInstances[i] = Spawn(IFA, Instances[i]);
+			ActorInstances[i] = Spawn(Instances[i]);
 		}
 	}
 }
@@ -243,12 +246,13 @@ void FFoliageActor::UpdateActorTransforms(const TArray<FFoliageInstance>& Instan
 	}
 }
 
-void FFoliageActor::PostEditUndo(AInstancedFoliageActor* IFA, UFoliageType* FoliageType, const TArray<FFoliageInstance>& Instances, const TSet<int32>& SelectedIndices)
+void FFoliageActor::PostEditUndo(FFoliageInfo* InInfo, UFoliageType* FoliageType, const TArray<FFoliageInstance>& Instances, const TSet<int32>& SelectedIndices)
 {
+	FFoliageImpl::PostEditUndo(InInfo, FoliageType, Instances, SelectedIndices);
 	UpdateActorTransforms(Instances);
 }
 
-void FFoliageActor::PreMoveInstances(AInstancedFoliageActor* InIFA, const TArray<int32>& InInstancesMoved)
+void FFoliageActor::PreMoveInstances(const TArray<int32>& InInstancesMoved)
 {
 	for (int32 Index : InInstancesMoved)
 	{
@@ -259,7 +263,7 @@ void FFoliageActor::PreMoveInstances(AInstancedFoliageActor* InIFA, const TArray
 	}
 }
 
-void FFoliageActor::PostMoveInstances(AInstancedFoliageActor* InIFA, const TArray<int32>& InInstancesMoved, bool bFinished)
+void FFoliageActor::PostMoveInstances(const TArray<int32>& InInstancesMoved, bool bFinished)
 {
 	// Copy because moving actors might remove them from ActorInstances
 	TArray<AActor*> MovedActors;
@@ -279,8 +283,9 @@ void FFoliageActor::PostMoveInstances(AInstancedFoliageActor* InIFA, const TArra
 	}
 }
 
-void FFoliageActor::NotifyFoliageTypeChanged(AInstancedFoliageActor* IFA, UFoliageType* FoliageType, const TArray<FFoliageInstance>& Instances, const TSet<int32>& SelectedIndices, bool bSourceChanged)
+void FFoliageActor::NotifyFoliageTypeChanged(UFoliageType* FoliageType, const TArray<FFoliageInstance>& Instances, const TSet<int32>& SelectedIndices, bool bSourceChanged)
 {
+	AInstancedFoliageActor* IFA = GetIFA();
 	if (UFoliageType_Actor* FoliageTypeActor = Cast<UFoliageType_Actor>(FoliageType))
 	{
 		if (bShouldAttachToBaseComponent != FoliageTypeActor->bShouldAttachToBaseComponent)
@@ -295,13 +300,14 @@ void FFoliageActor::NotifyFoliageTypeChanged(AInstancedFoliageActor* IFA, UFolia
 
 	if (bSourceChanged)
 	{
-		Reapply(IFA, FoliageType, Instances);
+		Reapply(FoliageType, Instances);
 		ApplySelection(true, SelectedIndices);
 	}
 }
 
-void FFoliageActor::Reapply(AInstancedFoliageActor* IFA, const UFoliageType* FoliageType, const TArray<FFoliageInstance>& Instances, bool bPostLoad)
+void FFoliageActor::Reapply(const UFoliageType* FoliageType, const TArray<FFoliageInstance>& Instances, bool bPostLoad)
 {
+	AInstancedFoliageActor* IFA = GetIFA();
 	IFA->Modify();
 	DestroyActors(bPostLoad);
 	
@@ -309,11 +315,11 @@ void FFoliageActor::Reapply(AInstancedFoliageActor* IFA, const UFoliageType* Fol
 	{
 		Uninitialize();
 	}
-	Initialize(IFA, FoliageType);
+	Initialize(FoliageType);
 
 	for (int32 i = 0; i < Instances.Num(); ++i)
 	{
-		ActorInstances.Add(Spawn(IFA, Instances[i]));
+		ActorInstances.Add(Spawn(Instances[i]));
 	}
 }
 
@@ -357,9 +363,10 @@ void FFoliageActor::ClearSelection(const TSet<int32>& SelectedIndices)
 	AInstancedFoliageActor::SelectionChanged.Broadcast(false, GetActorsFromSelectedIndices(SelectedIndices));
 }
 
-bool FFoliageActor::UpdateInstanceFromActor(AInstancedFoliageActor* IFA, int32 Index, FFoliageInfo& FoliageInfo)
+bool FFoliageActor::UpdateInstanceFromActor(int32 Index, FFoliageInfo& FoliageInfo)
 {
 	AActor* Actor = ActorInstances[Index];
+	AInstancedFoliageActor* IFA = GetIFA();
 	IFA->Modify();
 	const bool bChecked = false; // In the case of the PostEditUndo its possible that the instancehash is empty.
 	FoliageInfo.InstanceHash->RemoveInstance(FoliageInfo.Instances[Index].Location, Index, bChecked);
diff --git a/Engine/Source/Runtime/Foliage/Private/FoliageActor.h b/Engine/Source/Runtime/Foliage/Private/FoliageActor.h
--- a/Engine/Source/Runtime/Foliage/Private/FoliageActor.h
+++ b/Engine/Source/Runtime/Foliage/Private/FoliageActor.h
@@ -13,8 +13,9 @@ struct FFoliageActor : public FFoliageImpl
 	UClass* ActorClass;
 	bool bShouldAttachToBaseComponent;
 	
-	FFoliageActor()
-		: ActorClass(nullptr)
+	FFoliageActor(FFoliageInfo* Info)
+		: FFoliageImpl(Info)
+		, ActorClass(nullptr)
 		, bShouldAttachToBaseComponent(true)
 #if WITH_EDITOR
 		, bActorsDestroyed(false)
@@ -28,12 +29,12 @@ struct FFoliageActor : public FFoliageImpl
 
 #if WITH_EDITOR
 	virtual bool IsInitialized() const override;
-	virtual void Initialize(AInstancedFoliageActor* IFA, const UFoliageType* FoliageType) override;
+	virtual void Initialize(const UFoliageType* FoliageType) override;
 	virtual void Uninitialize() override;
 	virtual int32 GetInstanceCount() const override;
-	virtual void PreAddInstances(AInstancedFoliageActor* IFA, const UFoliageType* FoliageType, int32 Count) override;
-	virtual void AddInstance(AInstancedFoliageActor* IFA, const FFoliageInstance& NewInstance) override;
-	virtual void AddExistingInstance(AInstancedFoliageActor* IFA, const FFoliageInstance& ExistingInstance, UObject* InstanceImplementation) override;
+	virtual void PreAddInstances(const UFoliageType* FoliageType, int32 Count) override;
+	virtual void AddInstance(const FFoliageInstance& NewInstance) override;
+	virtual void AddExistingInstance(const FFoliageInstance& ExistingInstance, UObject* InstanceImplementation) override;
 	virtual void RemoveInstance(int32 InstanceIndex) override;
 	virtual void MoveInstance(int32 InstanceIndex, UObject*& OutInstanceImplementation) override;
 	virtual void SetInstanceWorldTransform(int32 InstanceIndex, const FTransform& Transform, bool bTeleport) override;
@@ -51,17 +52,17 @@ struct FFoliageActor : public FFoliageImpl
 
 	virtual void BeginUpdate() override;
 	virtual void EndUpdate() override;
-	virtual void Refresh(AInstancedFoliageActor* IFA, const TArray<FFoliageInstance>& Instances, bool Async, bool Force) override;
+	virtual void Refresh(const TArray<FFoliageInstance>& Instances, bool Async, bool Force) override;
 	virtual void OnHiddenEditorViewMaskChanged(uint64 InHiddenEditorViews) override;
-	virtual void PostEditUndo(AInstancedFoliageActor* IFA, UFoliageType* FoliageType, const TArray<FFoliageInstance>& Instances, const TSet<int32>& SelectedIndices) override;
-	virtual void PreMoveInstances(AInstancedFoliageActor* InIFA, const TArray<int32>& InInstancesMoved) override;
-	virtual void PostMoveInstances(AInstancedFoliageActor* InIFA, const TArray<int32>& InInstancesMoved, bool bFinished) override;
-	virtual void NotifyFoliageTypeChanged(AInstancedFoliageActor* IFA, UFoliageType* FoliageType, const TArray<FFoliageInstance>& Instances, const TSet<int32>& SelectedIndices, bool bSourceChanged) override;
-	void Reapply(AInstancedFoliageActor* IFA, const UFoliageType* FoliageType, const TArray<FFoliageInstance>& Instances, bool bPostLoad = false);
-	AActor* Spawn(AInstancedFoliageActor* IFA, const FFoliageInstance& Instance);
+	virtual void PostEditUndo(FFoliageInfo* InInfo, UFoliageType* FoliageType, const TArray<FFoliageInstance>& Instances, const TSet<int32>& SelectedIndices) override;
+	virtual void PreMoveInstances(const TArray<int32>& InInstancesMoved) override;
+	virtual void PostMoveInstances(const TArray<int32>& InInstancesMoved, bool bFinished) override;
+	virtual void NotifyFoliageTypeChanged(UFoliageType* FoliageType, const TArray<FFoliageInstance>& Instances, const TSet<int32>& SelectedIndices, bool bSourceChanged) override;
+	void Reapply(const UFoliageType* FoliageType, const TArray<FFoliageInstance>& Instances, bool bPostLoad = false);
+	AActor* Spawn(const FFoliageInstance& Instance);
 	TArray<AActor*> GetActorsFromSelectedIndices(const TSet<int32>& SelectedIndices) const;
 	virtual bool ShouldAttachToBaseComponent() const override { return bShouldAttachToBaseComponent; }
-	bool UpdateInstanceFromActor(AInstancedFoliageActor* IFA, int32 Index, FFoliageInfo& FoliageInfo);
+	bool UpdateInstanceFromActor(int32 Index, FFoliageInfo& FoliageInfo);
 	void GetInvalidInstances(TArray<int32>& InvalidInstances);
 
 private:
diff --git a/Engine/Source/Runtime/Foliage/Private/InstancedFoliage.cpp b/Engine/Source/Runtime/Foliage/Private/InstancedFoliage.cpp
--- a/Engine/Source/Runtime/Foliage/Private/InstancedFoliage.cpp
+++ b/Engine/Source/Runtime/Foliage/Private/InstancedFoliage.cpp
@@ -78,8 +78,8 @@ FCustomVersionRegistration GRegisterFoliageCustomVersion(FFoliageCustomVersion::
 ///
 struct FFoliageStaticMesh : public FFoliageImpl
 {
-	FFoliageStaticMesh(UHierarchicalInstancedStaticMeshComponent* InComponent)
-		: FFoliageImpl()
+	FFoliageStaticMesh(FFoliageInfo* Info, UHierarchicalInstancedStaticMeshComponent* InComponent)
+		: FFoliageImpl(Info)
 		, Component(InComponent)
 #if WITH_EDITOR
 		, UpdateDepth(0)
@@ -103,11 +103,11 @@ struct FFoliageStaticMesh : public FFoliageImpl
 
 #if WITH_EDITOR
 	virtual bool IsInitialized() const override { return Component != nullptr; }
-	virtual void Initialize(AInstancedFoliageActor* IFA, const UFoliageType* FoliageType) override;
+	virtual void Initialize(const UFoliageType* FoliageType) override;
 	virtual void Uninitialize() override;
 	virtual int32 GetInstanceCount() const override;
-	virtual void PreAddInstances(AInstancedFoliageActor* IFA, const UFoliageType* FoliageType, int32 Count) override;
-	virtual void AddInstance(AInstancedFoliageActor* IFA, const FFoliageInstance& NewInstance) override;
+	virtual void PreAddInstances(const UFoliageType* FoliageType, int32 Count) override;
+	virtual void AddInstance(const FFoliageInstance& NewInstance) override;
 	virtual void RemoveInstance(int32 InstanceIndex) override;
 	virtual void SetInstanceWorldTransform(int32 InstanceIndex, const FTransform& Transform, bool bTeleport) override;
 	virtual FTransform GetInstanceWorldTransform(int32 InstanceIndex) const override;
@@ -123,12 +123,12 @@ struct FFoliageStaticMesh : public FFoliageImpl
 
 	virtual void BeginUpdate() override;
 	virtual void EndUpdate() override;
-	virtual void Refresh(AInstancedFoliageActor* IFA, const TArray<FFoliageInstance>& Instances, bool Async, bool Force) override;
+	virtual void Refresh(const TArray<FFoliageInstance>& Instances, bool Async, bool Force) override;
 	virtual void OnHiddenEditorViewMaskChanged(uint64 InHiddenEditorViews) override;
-	virtual void PreEditUndo(AInstancedFoliageActor* IFA, UFoliageType* FoliageType) override;
-	virtual void PostEditUndo(AInstancedFoliageActor* IFA, UFoliageType* FoliageType, const TArray<FFoliageInstance>& Instances, const TSet<int32>& SelectedIndices) override;
-	virtual void NotifyFoliageTypeWillChange(AInstancedFoliageActor* IFA, UFoliageType* FoliageType) override;
-	virtual void NotifyFoliageTypeChanged(AInstancedFoliageActor* IFA, UFoliageType* FoliageType, const TArray<FFoliageInstance>& Instances, const TSet<int32>& SelectedIndices, bool bSourceChanged) override;
+	virtual void PreEditUndo(UFoliageType* FoliageType) override;
+	virtual void PostEditUndo(FFoliageInfo* InInfo, UFoliageType* FoliageType, const TArray<FFoliageInstance>& Instances, const TSet<int32>& SelectedIndices) override;
+	virtual void NotifyFoliageTypeWillChange(UFoliageType* FoliageType) override;
+	virtual void NotifyFoliageTypeChanged(UFoliageType* FoliageType, const TArray<FFoliageInstance>& Instances, const TSet<int32>& SelectedIndices, bool bSourceChanged) override;
 	virtual void EnterEditMode() override;
 	virtual void ExitEditMode() override;
 
@@ -142,9 +142,9 @@ struct FFoliageStaticMesh : public FFoliageImpl
 
 	void UpdateComponentSettings(const UFoliageType_InstancedStaticMesh* InSettings);
 	// Recreate the component if the FoliageType's ComponentClass doesn't match the Component's class
-	void CheckComponentClass(AInstancedFoliageActor* InIFA, const UFoliageType_InstancedStaticMesh* InSettings, const TArray<FFoliageInstance>& Instances, const TSet<int32>& SelectedIndices);
+	void CheckComponentClass(const UFoliageType_InstancedStaticMesh* InSettings, const TArray<FFoliageInstance>& Instances, const TSet<int32>& SelectedIndices);
 	void ReapplyInstancesToComponent(const TArray<FFoliageInstance>& Instances, const TSet<int32>& SelectedIndices);
-	void CreateNewComponent(AInstancedFoliageActor* InIFA, const UFoliageType* InSettings);
+	void CreateNewComponent(const UFoliageType* InSettings);
 };
 
 // Legacy (< FFoliageCustomVersion::CrossLevelBase) serializer
@@ -211,12 +211,12 @@ static void ConvertDeprecatedFoliageMeshes(
 #if WITH_EDITORONLY_DATA	
 	for (auto Pair : FoliageMeshesDeprecated)
 	{
-		auto& FoliageMesh = FoliageInfos.Add(Pair.Key);
+		auto& FoliageMesh = IFA->AddFoliageInfo(Pair.Key);
 		const auto& FoliageMeshDeprecated = Pair.Value;
 
 		// Old Foliage mesh is always static mesh (no actors)
 		FoliageMesh->Type = EFoliageImplType::StaticMesh;
-		FoliageMesh->Implementation.Reset(new FFoliageStaticMesh(FoliageMeshDeprecated->Component));
+		FoliageMesh->Implementation.Reset(new FFoliageStaticMesh(&FoliageMesh.Get(), FoliageMeshDeprecated->Component));
 		FoliageMesh->FoliageTypeUpdateGuid = FoliageMeshDeprecated->FoliageTypeUpdateGuid;
 
 		FoliageMesh->Instances.Reserve(FoliageMeshDeprecated->Instances.Num());
@@ -252,12 +252,12 @@ static void ConvertDeprecated2FoliageMeshes(
 #if WITH_EDITORONLY_DATA	
 	for (auto Pair : FoliageMeshesDeprecated)
 	{
-		auto& FoliageMesh = FoliageInfos.Add(Pair.Key);
+		auto& FoliageMesh = IFA->AddFoliageInfo(Pair.Key);
 		const auto& FoliageMeshDeprecated = Pair.Value;
 
 		// Old Foliage mesh is always static mesh (no actors)
 		FoliageMesh->Type = EFoliageImplType::StaticMesh;
-		FoliageMesh->Implementation.Reset(new FFoliageStaticMesh(FoliageMeshDeprecated->Component));
+		FoliageMesh->Implementation.Reset(new FFoliageStaticMesh(&FoliageMesh.Get(), FoliageMeshDeprecated->Component));
 		FoliageMesh->FoliageTypeUpdateGuid = FoliageMeshDeprecated->FoliageTypeUpdateGuid;
 
 		FoliageMesh->Instances.Reserve(FoliageMeshDeprecated->Instances.Num());
@@ -988,9 +988,9 @@ void FFoliageStaticMesh::GetOverlappingMeshCount(const FSphere& Sphere, TMap<USt
 }
 
 #if WITH_EDITOR
-void FFoliageStaticMesh::Initialize(AInstancedFoliageActor* IFA, const UFoliageType* FoliageType)
+void FFoliageStaticMesh::Initialize(const UFoliageType* FoliageType)
 {
-	CreateNewComponent(IFA, FoliageType);
+	CreateNewComponent(FoliageType);
 }
 
 void FFoliageStaticMesh::Uninitialize()
@@ -1020,11 +1020,11 @@ int32 FFoliageStaticMesh::GetInstanceCount() const
 	return 0;
 }
 
-void FFoliageStaticMesh::PreAddInstances(AInstancedFoliageActor* IFA, const UFoliageType* FoliageType, int32 Count)
+void FFoliageStaticMesh::PreAddInstances(const UFoliageType* FoliageType, int32 Count)
 {
 	if (!IsInitialized())
 	{
-		Initialize(IFA, FoliageType);
+		Initialize(FoliageType);
 		check(IsInitialized());
 	}
 	else
@@ -1039,7 +1039,7 @@ void FFoliageStaticMesh::PreAddInstances(AInstancedFoliageActor* IFA, const UFol
 	}
 }
 
-void FFoliageStaticMesh::AddInstance(AInstancedFoliageActor* IFA, const FFoliageInstance& NewInstance)
+void FFoliageStaticMesh::AddInstance(const FFoliageInstance& NewInstance)
 {
 	check(Component);
 	Component->AddInstanceWorldSpace(NewInstance.GetInstanceWorldTransform());
@@ -1173,7 +1173,7 @@ void FFoliageStaticMesh::EndUpdate()
 	}
 }
 
-void FFoliageStaticMesh::Refresh(AInstancedFoliageActor* IFA, const TArray<FFoliageInstance>& Instances, bool Async, bool Force)
+void FFoliageStaticMesh::Refresh(const TArray<FFoliageInstance>& Instances, bool Async, bool Force)
 {
 	if (Component != nullptr)
 	{
@@ -1192,7 +1192,7 @@ void FFoliageStaticMesh::OnHiddenEditorViewMaskChanged(uint64 InHiddenEditorView
 	}
 }
 
-void FFoliageStaticMesh::PreEditUndo(AInstancedFoliageActor* IFA, UFoliageType* FoliageType)
+void FFoliageStaticMesh::PreEditUndo(UFoliageType* FoliageType)
 {
 	if (UFoliageType_InstancedStaticMesh* FoliageType_InstancedStaticMesh = Cast<UFoliageType_InstancedStaticMesh>(FoliageType))
 	{
@@ -1203,8 +1203,9 @@ void FFoliageStaticMesh::PreEditUndo(AInstancedFoliageActor* IFA, UFoliageType*
 	}
 }
 
-void FFoliageStaticMesh::PostEditUndo(AInstancedFoliageActor* IFA, UFoliageType* FoliageType, const TArray<FFoliageInstance>& Instances, const TSet<int32>& SelectedIndices)
+void FFoliageStaticMesh::PostEditUndo(FFoliageInfo* InInfo, UFoliageType* FoliageType, const TArray<FFoliageInstance>& Instances, const TSet<int32>& SelectedIndices)
 {
+	FFoliageImpl::PostEditUndo(InInfo, FoliageType, Instances, SelectedIndices);
 	if (UFoliageType_InstancedStaticMesh* FoliageType_InstancedStaticMesh = Cast<UFoliageType_InstancedStaticMesh>(FoliageType))
 	{
 		if (Component != nullptr && FoliageType_InstancedStaticMesh->GetStaticMesh() != nullptr)
@@ -1212,12 +1213,12 @@ void FFoliageStaticMesh::PostEditUndo(AInstancedFoliageActor* IFA, UFoliageType*
 			FoliageType_InstancedStaticMesh->GetStaticMesh()->GetOnExtendedBoundsChanged().AddRaw(this, &FFoliageStaticMesh::HandleComponentMeshBoundsChanged);
 		}
 
-		CheckComponentClass(IFA, FoliageType_InstancedStaticMesh, Instances, SelectedIndices);
+		CheckComponentClass(FoliageType_InstancedStaticMesh, Instances, SelectedIndices);
 		ReapplyInstancesToComponent(Instances, SelectedIndices);
 	}
 }
 
-void FFoliageStaticMesh::NotifyFoliageTypeWillChange(AInstancedFoliageActor* IFA, UFoliageType* FoliageType)
+void FFoliageStaticMesh::NotifyFoliageTypeWillChange(UFoliageType* FoliageType)
 {
 	if (Component != nullptr)
 	{
@@ -1231,11 +1232,11 @@ void FFoliageStaticMesh::NotifyFoliageTypeWillChange(AInstancedFoliageActor* IFA
 	}
 }
 
-void FFoliageStaticMesh::NotifyFoliageTypeChanged(AInstancedFoliageActor* IFA, UFoliageType* FoliageType, const TArray<FFoliageInstance>& Instances, const TSet<int32>& SelectedIndices, bool bSourceChanged)
+void FFoliageStaticMesh::NotifyFoliageTypeChanged(UFoliageType* FoliageType, const TArray<FFoliageInstance>& Instances, const TSet<int32>& SelectedIndices, bool bSourceChanged)
 {
 	UFoliageType_InstancedStaticMesh* FoliageType_InstancedStaticMesh = Cast<UFoliageType_InstancedStaticMesh>(FoliageType);
 	check(FoliageType_InstancedStaticMesh);
-	CheckComponentClass(IFA, FoliageType_InstancedStaticMesh, Instances, SelectedIndices);
+	CheckComponentClass(FoliageType_InstancedStaticMesh, Instances, SelectedIndices);
 	UpdateComponentSettings(FoliageType_InstancedStaticMesh);
 	
 	if (bSourceChanged && Component != nullptr && Component->GetStaticMesh() != nullptr)
@@ -1285,7 +1286,7 @@ void FFoliageStaticMesh::ExitEditMode()
 	Component->UpdateDensityScaling();
 }
 
-void FFoliageStaticMesh::CreateNewComponent(AInstancedFoliageActor* InIFA, const UFoliageType* InSettings)
+void FFoliageStaticMesh::CreateNewComponent(const UFoliageType* InSettings)
 {
 	SCOPE_CYCLE_COUNTER(STAT_FoliageCreateComponent);
 
@@ -1298,7 +1299,8 @@ void FFoliageStaticMesh::CreateNewComponent(AInstancedFoliageActor* InIFA, const
 		ComponentClass = UFoliageInstancedStaticMeshComponent::StaticClass();
 	}
 
-	UFoliageInstancedStaticMeshComponent* FoliageComponent = NewObject<UFoliageInstancedStaticMeshComponent>(InIFA, ComponentClass, NAME_None, RF_Transactional);
+	AInstancedFoliageActor* IFA = GetIFA();
+	UFoliageInstancedStaticMeshComponent* FoliageComponent = NewObject<UFoliageInstancedStaticMeshComponent>(IFA, ComponentClass, NAME_None, RF_Transactional);
 	
 	check(FoliageType_InstancedStaticMesh);
 
@@ -1316,15 +1318,15 @@ void FFoliageStaticMesh::CreateNewComponent(AInstancedFoliageActor* InIFA, const
 
 	UpdateComponentSettings(FoliageType_InstancedStaticMesh);
 
-	Component->SetupAttachment(InIFA->GetRootComponent());
+	Component->SetupAttachment(IFA->GetRootComponent());
 
-	if (InIFA->GetRootComponent()->IsRegistered())
+	if (IFA->GetRootComponent()->IsRegistered())
 	{
 		Component->RegisterComponent();
 	}
 
 	// Use only instance translation as a component transform
-	Component->SetWorldTransform(InIFA->GetRootComponent()->GetComponentTransform());
+	Component->SetWorldTransform(IFA->GetRootComponent()->GetComponentTransform());
 
 	// Add the new component to the transaction buffer so it will get destroyed on undo
 	Component->Modify();
@@ -1340,7 +1342,7 @@ void FFoliageStaticMesh::HandleComponentMeshBoundsChanged(const FBoxSphereBounds
 	}
 }
 
-void FFoliageStaticMesh::CheckComponentClass(AInstancedFoliageActor* InIFA, const UFoliageType_InstancedStaticMesh* InSettings, const TArray<FFoliageInstance>& Instances, const TSet<int32>& SelectedIndices)
+void FFoliageStaticMesh::CheckComponentClass(const UFoliageType_InstancedStaticMesh* InSettings, const TArray<FFoliageInstance>& Instances, const TSet<int32>& SelectedIndices)
 {
 	if (Component)
 	{
@@ -1352,13 +1354,14 @@ void FFoliageStaticMesh::CheckComponentClass(AInstancedFoliageActor* InIFA, cons
 
 		if (ComponentClass != Component->GetClass())
 		{
-			InIFA->Modify();
+			AInstancedFoliageActor* IFA = GetIFA();
+			IFA->Modify();
 
 			// prepare to destroy the old component
 			Uninitialize();
 
 			// create a new component
-			Initialize(InIFA, InSettings);
+			Initialize(InSettings);
 
 			// apply the instances to it
 			ReapplyInstancesToComponent(Instances, SelectedIndices);
@@ -1623,6 +1626,11 @@ void FFoliageStaticMesh::ReapplyInstancesToComponent(const TArray<FFoliageInstan
 	}
 }
 
+AInstancedFoliageActor* FFoliageImpl::GetIFA() const
+{
+	return Info->IFA;
+}
+
 #endif // WITH_EDITOR
 
 //
@@ -1630,7 +1638,8 @@ void FFoliageStaticMesh::ReapplyInstancesToComponent(const TArray<FFoliageInstan
 //
 FFoliageInfo::FFoliageInfo()
 	: Type(EFoliageImplType::StaticMesh)
-#if WITH_EDITOR
+#if WITH_EDITORONLY_DATA
+	, IFA(nullptr)
 	, InstanceHash(GIsEditor ? new FFoliageInstanceHash() : nullptr)
 	, bMovingInstances(false)
 #endif
@@ -1667,11 +1676,11 @@ void FFoliageInfo::CreateImplementation(EFoliageImplType InType)
 	
 	if (Type == EFoliageImplType::StaticMesh)
 	{
-		Implementation.Reset(new FFoliageStaticMesh(nullptr));
+		Implementation.Reset(new FFoliageStaticMesh(this, nullptr));
 	}
 	else if (Type == EFoliageImplType::Actor)
 	{
-		Implementation.Reset(new FFoliageActor());
+		Implementation.Reset(new FFoliageActor(this));
 	}
 }
 
@@ -1716,25 +1725,23 @@ void FFoliageInfo::CreateImplementation(const UFoliageType* FoliageType)
 {
 	check(!Implementation.IsValid());
 	// Change Impl based on FoliageType param
-	Type = EFoliageImplType::Unknown;
+	EFoliageImplType ImplType = EFoliageImplType::Unknown;
 	if (FoliageType->IsA<UFoliageType_InstancedStaticMesh>())
 	{
-		Type = EFoliageImplType::StaticMesh;
-		Implementation.Reset(new FFoliageStaticMesh(nullptr));
+		ImplType = EFoliageImplType::StaticMesh;
 	}
 	else if (FoliageType->IsA<UFoliageType_Actor>())
 	{
-		Type = EFoliageImplType::Actor;
-		Implementation.Reset(new FFoliageActor());
+		ImplType = EFoliageImplType::Actor;
 	}
-	check(Type != EFoliageImplType::Unknown);
+	CreateImplementation(ImplType);
 }
 
-void FFoliageInfo::Initialize(AInstancedFoliageActor* IFA, const UFoliageType* FoliageType)
+void FFoliageInfo::Initialize(const UFoliageType* FoliageType)
 {
 	check(!IsInitialized());
 	check(Implementation.IsValid());
-	Implementation->Initialize(IFA, FoliageType);
+	Implementation->Initialize(FoliageType);
 }
 
 void FFoliageInfo::Uninitialize()
@@ -1748,14 +1755,14 @@ bool FFoliageInfo::IsInitialized() const
 	return Implementation.IsValid() && Implementation->IsInitialized();
 }
 
-void FFoliageInfo::NotifyFoliageTypeWillChange(AInstancedFoliageActor* IFA, UFoliageType* FoliageType)
+void FFoliageInfo::NotifyFoliageTypeWillChange(UFoliageType* FoliageType)
 {
-	Implementation->NotifyFoliageTypeWillChange(IFA, FoliageType);
+	Implementation->NotifyFoliageTypeWillChange(FoliageType);
 }
 
-void FFoliageInfo::NotifyFoliageTypeChanged(AInstancedFoliageActor* IFA, UFoliageType* FoliageType, bool bSourceChanged)
+void FFoliageInfo::NotifyFoliageTypeChanged(UFoliageType* FoliageType, bool bSourceChanged)
 {
-	Implementation->NotifyFoliageTypeChanged(IFA, FoliageType, Instances, SelectedIndices, bSourceChanged);
+	Implementation->NotifyFoliageTypeChanged(FoliageType, Instances, SelectedIndices, bSourceChanged);
 }
 
 void FFoliageInfo::CheckValid()
@@ -1823,13 +1830,13 @@ void FFoliageInfo::SetInstanceWorldTransform(int32 InstanceIndex, const FTransfo
 	Implementation->SetInstanceWorldTransform(InstanceIndex, Transform, bTeleport);
 }
 
-void FFoliageInfo::AddInstanceImpl(AInstancedFoliageActor* InIFA, const FFoliageInstance& InNewInstance, FFoliageInfo::FAddImplementationFunc ImplementationFunc)
+void FFoliageInfo::AddInstanceImpl(const FFoliageInstance& InNewInstance, FFoliageInfo::FAddImplementationFunc ImplementationFunc)
 {
 	// Add the instance taking either a free slot or adding a new item.
 	int32 InstanceIndex = Instances.Add(InNewInstance);
 	FFoliageInstance& AddedInstance = Instances[InstanceIndex];
 
-	AddedInstance.BaseId = InIFA->InstanceBaseCache.AddInstanceBaseId(ShouldAttachToBaseComponent()? InNewInstance.BaseComponent : nullptr);
+	AddedInstance.BaseId = IFA->InstanceBaseCache.AddInstanceBaseId(ShouldAttachToBaseComponent()? InNewInstance.BaseComponent : nullptr);
 	if (AddedInstance.BaseId == FFoliageInstanceBaseCache::InvalidBaseId)
 	{
 		AddedInstance.BaseComponent = nullptr;
@@ -1840,27 +1847,27 @@ void FFoliageInfo::AddInstanceImpl(AInstancedFoliageActor* InIFA, const FFoliage
 	InstanceHash->InsertInstance(AddedInstance.Location, InstanceIndex);
 
 	// Add the instance to the component
-	ImplementationFunc(Implementation.Get(), InIFA, AddedInstance);
+	ImplementationFunc(Implementation.Get(), IFA, AddedInstance);
 }
 
-void FFoliageInfo::AddInstances(AInstancedFoliageActor* InIFA, const UFoliageType* InSettings, const TArray<const FFoliageInstance*>& InNewInstances)
+void FFoliageInfo::AddInstances(const UFoliageType* InSettings, const TArray<const FFoliageInstance*>& InNewInstances)
 {
-	AddInstancesImpl(InIFA, InSettings, InNewInstances, [](FFoliageImpl* Impl, AInstancedFoliageActor* LocalIFA, const FFoliageInstance& LocalInstance) { Impl->AddInstance(LocalIFA, LocalInstance); });
+	AddInstancesImpl(InSettings, InNewInstances, [](FFoliageImpl* Impl, AInstancedFoliageActor* LocalIFA, const FFoliageInstance& LocalInstance) { Impl->AddInstance(LocalInstance); });
 }
 
-void FFoliageInfo::ReserveAdditionalInstances(AInstancedFoliageActor* InIFA, const UFoliageType* InSettings, uint32 ReserveNum)
+void FFoliageInfo::ReserveAdditionalInstances(const UFoliageType* InSettings, uint32 ReserveNum)
 {
 	Instances.Reserve(Instances.Num() + ReserveNum);
-	Implementation->PreAddInstances(InIFA, InSettings, ReserveNum);
+	Implementation->PreAddInstances(InSettings, ReserveNum);
 }
 
-void FFoliageInfo::AddInstancesImpl(AInstancedFoliageActor* InIFA, const UFoliageType* InSettings, const TArray<const FFoliageInstance*>& InNewInstances, FFoliageInfo::FAddImplementationFunc ImplementationFunc)
+void FFoliageInfo::AddInstancesImpl(const UFoliageType* InSettings, const TArray<const FFoliageInstance*>& InNewInstances, FFoliageInfo::FAddImplementationFunc ImplementationFunc)
 {
 	SCOPE_CYCLE_COUNTER(STAT_FoliageAddInstance);
 
-	InIFA->Modify();
+	IFA->Modify();
 
-	Implementation->PreAddInstances(InIFA, InSettings, InNewInstances.Num());
+	Implementation->PreAddInstances(InSettings, InNewInstances.Num());
 
 	Implementation->BeginUpdate();
 
@@ -1868,7 +1875,7 @@ void FFoliageInfo::AddInstancesImpl(AInstancedFoliageActor* InIFA, const UFoliag
 
 	for (const FFoliageInstance* Instance : InNewInstances)
 	{
-		AddInstanceImpl(InIFA, *Instance, ImplementationFunc);
+		AddInstanceImpl(*Instance, ImplementationFunc);
 	}
 
 	CheckValid();
@@ -1876,28 +1883,28 @@ void FFoliageInfo::AddInstancesImpl(AInstancedFoliageActor* InIFA, const UFoliag
 	Implementation->EndUpdate();
 }
 
-void FFoliageInfo::AddInstance(AInstancedFoliageActor* InIFA, const UFoliageType* InSettings, const FFoliageInstance& InNewInstance, UActorComponent* InBaseComponent)
+void FFoliageInfo::AddInstance(const UFoliageType* InSettings, const FFoliageInstance& InNewInstance, UActorComponent* InBaseComponent)
 {
 	FFoliageInstance Instance = InNewInstance;
-	Instance.BaseId = InIFA->InstanceBaseCache.AddInstanceBaseId(ShouldAttachToBaseComponent() ? InBaseComponent : nullptr);
+	Instance.BaseId = IFA->InstanceBaseCache.AddInstanceBaseId(ShouldAttachToBaseComponent() ? InBaseComponent : nullptr);
 	if (Instance.BaseId == FFoliageInstanceBaseCache::InvalidBaseId)
 	{
 		Instance.BaseComponent = nullptr;
 	}
-	AddInstance(InIFA, InSettings, Instance);
+	AddInstance(InSettings, Instance);
 }
 
-void FFoliageInfo::AddInstance(AInstancedFoliageActor* InIFA, const UFoliageType* InSettings, const FFoliageInstance& InNewInstance)
+void FFoliageInfo::AddInstance(const UFoliageType* InSettings, const FFoliageInstance& InNewInstance)
 {
-	AddInstances(InIFA, InSettings, { &InNewInstance });
+	AddInstances(InSettings, { &InNewInstance });
 }
 
-void FFoliageInfo::MoveInstances(AInstancedFoliageActor* InFromIFA, AInstancedFoliageActor* InToIFA, const TSet<int32>& InInstancesToMove, bool bKeepSelection)
+void FFoliageInfo::MoveInstances(AInstancedFoliageActor* InToIFA, const TSet<int32>& InInstancesToMove, bool bKeepSelection)
 {
-	const UFoliageType* FoliageType = InFromIFA->FindFoliageType(this);
+	const UFoliageType* FoliageType = IFA->FindFoliageType(this);
 	check(FoliageType);
 
-	InFromIFA->Modify();
+	IFA->Modify();
 	
 	FFoliageInfo* OutFoliageInfo = nullptr;
 	UFoliageType* ToFoliageType = nullptr;
@@ -1925,14 +1932,14 @@ void FFoliageInfo::MoveInstances(AInstancedFoliageActor* InFromIFA, AInstancedFo
 	
 	for (int32 InstanceIndex : InInstancesToMove)
 	{
-		MoveData.Add(InstanceIndex, FFoliageMoveInstance(Instances[InstanceIndex], InFromIFA->GetBaseComponentFromBaseId(Instances[InstanceIndex].BaseId)));
+		MoveData.Add(InstanceIndex, FFoliageMoveInstance(Instances[InstanceIndex], IFA->GetBaseComponentFromBaseId(Instances[InstanceIndex].BaseId)));
 		if (bKeepSelection && OutFoliageInfo && SelectedIndices.Contains(InstanceIndex))
 		{
 			NewSelectedIndices.Add(OutFoliageInfo->Instances.Num() + MoveData.Num() - 1);
 		}
 	}
 
-	RemoveInstancesImpl(InFromIFA, InstancesToMove, true, [&MoveData](FFoliageImpl* Impl, int32 Index)
+	RemoveInstancesImpl(InstancesToMove, true, [&MoveData](FFoliageImpl* Impl, int32 Index)
 	{
 		Impl->MoveInstance(Index, MoveData[Index].InstanceImplementation);
 	});
@@ -1948,29 +1955,29 @@ void FFoliageInfo::MoveInstances(AInstancedFoliageActor* InFromIFA, AInstancedFo
 	{
 		int32 AddedIndex = 0;
 		
-		OutFoliageInfo->AddInstancesImpl(InToIFA, ToFoliageType, NewInstances, [&](FFoliageImpl* Impl, AInstancedFoliageActor* LocalIFA, const FFoliageInstance& LocalInstance)
+		OutFoliageInfo->AddInstancesImpl(ToFoliageType, NewInstances, [&](FFoliageImpl* Impl, AInstancedFoliageActor* LocalIFA, const FFoliageInstance& LocalInstance)
 		{
 			const FFoliageMoveInstance* MoveInstance = StaticCast<const FFoliageMoveInstance*>(NewInstances[AddedIndex]);
-			Impl->AddExistingInstance(LocalIFA, LocalInstance, MoveInstance->InstanceImplementation);
+			Impl->AddExistingInstance(LocalInstance, MoveInstance->InstanceImplementation);
 			AddedIndex++;
 		});
 
-		OutFoliageInfo->Refresh(InToIFA, true, true);
+		OutFoliageInfo->Refresh(true, true);
 
 		// Select if needed
 		if (NewSelectedIndices.Num())
 		{
-			OutFoliageInfo->SelectInstances(InToIFA, true, NewSelectedIndices);
+			OutFoliageInfo->SelectInstances(true, NewSelectedIndices);
 		}
 	}
 }
 
-void FFoliageInfo::RemoveInstances(AInstancedFoliageActor* InIFA, const TArray<int32>& InInstancesToRemove, bool RebuildFoliageTree)
+void FFoliageInfo::RemoveInstances(const TArray<int32>& InInstancesToRemove, bool RebuildFoliageTree)
 {
-	RemoveInstancesImpl(InIFA, InInstancesToRemove, RebuildFoliageTree, [](FFoliageImpl* Impl, int32 Index) { Impl->RemoveInstance(Index); });
+	RemoveInstancesImpl(InInstancesToRemove, RebuildFoliageTree, [](FFoliageImpl* Impl, int32 Index) { Impl->RemoveInstance(Index); });
 }
 
-void FFoliageInfo::RemoveInstancesImpl(AInstancedFoliageActor* InIFA, const TArray<int32>& InInstancesToRemove, bool RebuildFoliageTree, FFoliageInfo::FRemoveImplementationFunc ImplementationFunc)
+void FFoliageInfo::RemoveInstancesImpl(const TArray<int32>& InInstancesToRemove, bool RebuildFoliageTree, FFoliageInfo::FRemoveImplementationFunc ImplementationFunc)
 {
 	SCOPE_CYCLE_COUNTER(STAT_FoliageRemoveInstance);
 
@@ -1980,7 +1987,7 @@ void FFoliageInfo::RemoveInstancesImpl(AInstancedFoliageActor* InIFA, const TArr
 	}
 
 	check(IsInitialized());
-	InIFA->Modify();
+	IFA->Modify();
 
 	Implementation->BeginUpdate();
 
@@ -2051,13 +2058,13 @@ void FFoliageInfo::RemoveInstancesImpl(AInstancedFoliageActor* InIFA, const TArr
 		
 	if (RebuildFoliageTree)
 	{
-		Refresh(InIFA, true, true);
+		Refresh(true, true);
 	}
 
 	CheckValid();
 }
 
-void FFoliageInfo::PreMoveInstances(AInstancedFoliageActor* InIFA, const TArray<int32>& InInstancesToMove)
+void FFoliageInfo::PreMoveInstances(const TArray<int32>& InInstancesToMove)
 {
 	bMovingInstances = true;
 
@@ -2069,11 +2076,11 @@ void FFoliageInfo::PreMoveInstances(AInstancedFoliageActor* InIFA, const TArray<
 		InstanceHash->RemoveInstance(Instance.Location, InstanceIndex);
 	}
 
-	Implementation->PreMoveInstances(InIFA, InInstancesToMove);
+	Implementation->PreMoveInstances(InInstancesToMove);
 }
 
 
-void FFoliageInfo::PostUpdateInstances(AInstancedFoliageActor* InIFA, const TArray<int32>& InInstancesUpdated, bool bReAddToHash, bool InUpdateSelection)
+void FFoliageInfo::PostUpdateInstances(const TArray<int32>& InInstancesUpdated, bool bReAddToHash, bool InUpdateSelection)
 {
 	if (InInstancesUpdated.Num())
 	{
@@ -2110,25 +2117,25 @@ void FFoliageInfo::PostUpdateInstances(AInstancedFoliageActor* InIFA, const TArr
 	}
 }
 
-void FFoliageInfo::PostMoveInstances(AInstancedFoliageActor* InIFA, const TArray<int32>& InInstancesMoved, bool bFinished)
+void FFoliageInfo::PostMoveInstances(const TArray<int32>& InInstancesMoved, bool bFinished)
 {
-	PostUpdateInstances(InIFA, InInstancesMoved, true, true);
-	Implementation->PostMoveInstances(InIFA, InInstancesMoved, bFinished);
+	PostUpdateInstances(InInstancesMoved, true, true);
+	Implementation->PostMoveInstances(InInstancesMoved, bFinished);
 	bMovingInstances = false;
 }
 
-void FFoliageInfo::DuplicateInstances(AInstancedFoliageActor* InIFA, UFoliageType* InSettings, const TArray<int32>& InInstancesToDuplicate)
+void FFoliageInfo::DuplicateInstances(UFoliageType* InSettings, const TArray<int32>& InInstancesToDuplicate)
 {
 	Implementation->BeginUpdate();
 
 	for (int32 InstanceIndex : InInstancesToDuplicate)
 	{
 		const FFoliageInstance TempInstance = Instances[InstanceIndex];
-		AddInstance(InIFA, InSettings, TempInstance);
+		AddInstance(InSettings, TempInstance);
 	}
 
 	Implementation->EndUpdate();
-	Refresh(InIFA, true, true);
+	Refresh(true, true);
 }
 
 /* Get the number of placed instances */
@@ -2171,7 +2178,7 @@ void FFoliageInfo::RemoveFromBaseHash(int32 InstanceIndex)
 }
 
 // Destroy existing clusters and reassign all instances to new clusters
-void FFoliageInfo::ReallocateClusters(AInstancedFoliageActor* InIFA, UFoliageType* InSettings)
+void FFoliageInfo::ReallocateClusters(UFoliageType* InSettings)
 {
 	// In case Foliage Type Changed recreate implementation
 	Implementation.Reset();
@@ -2192,15 +2199,15 @@ void FFoliageInfo::ReallocateClusters(AInstancedFoliageActor* InIFA, UFoliageTyp
 	{
 		if ((Instance.Flags & FOLIAGE_InstanceDeleted) == 0)
 		{
-			AddInstance(InIFA, InSettings, Instance);
+			AddInstance(InSettings, Instance);
 		}
 	}
 
 	
-	Refresh(InIFA, true, true);
+	Refresh(true, true);
 }
 
-void FFoliageInfo::GetInstancesInsideBounds(const FBox& Box, TArray<int32>& OutInstances)
+void FFoliageInfo::GetInstancesInsideBounds(const FBox& Box, TArray<int32>& OutInstances) const
 {
 	auto TempInstances = InstanceHash->GetInstancesOverlappingBox(Box);
 	for (int32 Idx : TempInstances)
@@ -2212,7 +2219,7 @@ void FFoliageInfo::GetInstancesInsideBounds(const FBox& Box, TArray<int32>& OutI
 	}
 }
 
-void FFoliageInfo::GetInstancesInsideSphere(const FSphere& Sphere, TArray<int32>& OutInstances)
+void FFoliageInfo::GetInstancesInsideSphere(const FSphere& Sphere, TArray<int32>& OutInstances) const
 {
 	auto TempInstances = InstanceHash->GetInstancesOverlappingBox(FBox::BuildAABB(Sphere.Center, FVector(Sphere.W)));
 	for (int32 Idx : TempInstances)
@@ -2224,7 +2231,7 @@ void FFoliageInfo::GetInstancesInsideSphere(const FSphere& Sphere, TArray<int32>
 	}
 }
 
-void FFoliageInfo::GetInstanceAtLocation(const FVector& Location, int32& OutInstance, bool& bOutSucess)
+void FFoliageInfo::GetInstanceAtLocation(const FVector& Location, int32& OutInstance, bool& bOutSucess) const
 {
 	auto TempInstances = InstanceHash->GetInstancesOverlappingBox(FBox::BuildAABB(Location, FVector(KINDA_SMALL_NUMBER)));
 
@@ -2246,7 +2253,7 @@ void FFoliageInfo::GetInstanceAtLocation(const FVector& Location, int32& OutInst
 }
 
 // Returns whether or not there is are any instances overlapping the sphere specified
-bool FFoliageInfo::CheckForOverlappingSphere(const FSphere& Sphere)
+bool FFoliageInfo::CheckForOverlappingSphere(const FSphere& Sphere) const
 {
 	auto TempInstances = InstanceHash->GetInstancesOverlappingBox(FBox::BuildAABB(Sphere.Center, FVector(Sphere.W)));
 	for (int32 Idx : TempInstances)
@@ -2260,7 +2267,7 @@ bool FFoliageInfo::CheckForOverlappingSphere(const FSphere& Sphere)
 }
 
 // Returns whether or not there is are any instances overlapping the instance specified, excluding the set of instances provided
-bool FFoliageInfo::CheckForOverlappingInstanceExcluding(int32 TestInstanceIdx, float Radius, TSet<int32>& ExcludeInstances)
+bool FFoliageInfo::CheckForOverlappingInstanceExcluding(int32 TestInstanceIdx, float Radius, TSet<int32>& ExcludeInstances) const
 {
 	FSphere Sphere(Instances[TestInstanceIdx].Location, Radius);
 
@@ -2275,11 +2282,11 @@ bool FFoliageInfo::CheckForOverlappingInstanceExcluding(int32 TestInstanceIdx, f
 	return false;
 }
 
-void FFoliageInfo::SelectInstances(AInstancedFoliageActor* InIFA, bool bSelect)
+void FFoliageInfo::SelectInstances(bool bSelect)
 {
 	if (Implementation->IsInitialized())
 	{
-		InIFA->Modify(false);
+		IFA->Modify(false);
 
 		if (bSelect)
 		{
@@ -2306,17 +2313,16 @@ FBox FFoliageInfo::GetSelectionBoundingBox() const
 	return Implementation->GetSelectionBoundingBox(SelectedIndices);
 }
 
-void FFoliageInfo::SelectInstances(AInstancedFoliageActor* InIFA, bool bSelect, TArray<int32>& InInstances)
+void FFoliageInfo::SelectInstances(bool bSelect, TArray<int32>& InInstances)
 {
 	if (InInstances.Num())
 	{
 		TSet<int32> ModifiedSelection;
 		ModifiedSelection.Reserve(InInstances.Num());
 		check(Implementation->IsInitialized());
+		IFA->Modify(false);
 		if (bSelect)
 		{
-			InIFA->Modify(false);
-
 			SelectedIndices.Reserve(InInstances.Num());
 
 			for (int32 i : InInstances)
@@ -2327,8 +2333,6 @@ void FFoliageInfo::SelectInstances(AInstancedFoliageActor* InIFA, bool bSelect,
 		}
 		else
 		{
-			InIFA->Modify(false);
-
 			for (int32 i : InInstances)
 			{
 				SelectedIndices.Remove(i);
@@ -2340,10 +2344,10 @@ void FFoliageInfo::SelectInstances(AInstancedFoliageActor* InIFA, bool bSelect,
 	}
 }
 
-void FFoliageInfo::Refresh(AInstancedFoliageActor* IFA, bool Async, bool Force)
+void FFoliageInfo::Refresh(bool Async, bool Force)
 {
 	check(Implementation.IsValid())
-	Implementation->Refresh(IFA, this->Instances, Async, Force);
+	Implementation->Refresh(this->Instances, Async, Force);
 }
 
 void FFoliageInfo::OnHiddenEditorViewMaskChanged(uint64 InHiddenEditorViews)
@@ -2351,14 +2355,16 @@ void FFoliageInfo::OnHiddenEditorViewMaskChanged(uint64 InHiddenEditorViews)
 	Implementation->OnHiddenEditorViewMaskChanged(InHiddenEditorViews);
 }
 
-void FFoliageInfo::PreEditUndo(AInstancedFoliageActor* IFA, UFoliageType* FoliageType)
+void FFoliageInfo::PreEditUndo(UFoliageType* FoliageType)
 {
-	Implementation->PreEditUndo(IFA, FoliageType);
+	Implementation->PreEditUndo(FoliageType);
 }
 
-void FFoliageInfo::PostEditUndo(AInstancedFoliageActor* IFA, UFoliageType* FoliageType)
+void FFoliageInfo::PostEditUndo(AInstancedFoliageActor* InIFA, UFoliageType* FoliageType)
 {
-	Implementation->PostEditUndo(IFA, FoliageType, Instances, SelectedIndices);
+	// Set the IFA after PostEditUndo as it is not a serialized member (will be nulled on serialization)
+	IFA = InIFA;
+	Implementation->PostEditUndo(this, FoliageType, Instances, SelectedIndices);
 		
 	// Regenerate instance hash
 	// We regenerate it here instead of saving to transaction buffer to speed up modify operations
@@ -2390,7 +2396,7 @@ void FFoliageInfo::RemoveBaseComponentOnInstances()
 	}
 }
 
-void FFoliageInfo::IncludeActor(AInstancedFoliageActor* IFA, const UFoliageType* FoliageType, AActor* InActor)
+void FFoliageInfo::IncludeActor(const UFoliageType* FoliageType, AActor* InActor)
 {
 	if (FFoliageActor* FoliageActor = StaticCast<FFoliageActor*>(Implementation.Get()))
 	{
@@ -2408,7 +2414,7 @@ void FFoliageInfo::IncludeActor(AInstancedFoliageActor* IFA, const UFoliageType*
 
 			int32 Index = Instances.Add(NewInstance);
 			InstanceHash->InsertInstance(NewInstance.Location, Index);
-			FoliageActor->PreAddInstances(IFA, FoliageType, 1);
+			FoliageActor->PreAddInstances(FoliageType, 1);
 			FoliageActor->ActorInstances.Add(InActor);
 			InActor->Modify();
 			FFoliageHelper::SetIsOwnedByFoliage(InActor);
@@ -2530,7 +2536,35 @@ void AInstancedFoliageActor::GetOverlappingMeshCounts(const FSphere& Sphere, TMa
 	}
 }
 
+void AInstancedFoliageActor::ForEachFoliageInfo(TFunctionRef<bool(UFoliageType* FoliageType, FFoliageInfo& FoliageInfo)> InOperation)
+{
+	for (auto& Pair : FoliageInfos)
+	{
+		if (!InOperation(Pair.Key, Pair.Value.Get()))
+		{
+			return;
+		}
+	}
+}
 
+TUniqueObj<FFoliageInfo>& AInstancedFoliageActor::AddFoliageInfo(UFoliageType* FoliageType)
+{
+	TUniqueObj<FFoliageInfo>& NewFoliageInfo = FoliageInfos.Add(FoliageType);
+	NewFoliageInfo->IFA = this;
+	return NewFoliageInfo;
+}
+
+TUniqueObj<FFoliageInfo>& AInstancedFoliageActor::AddFoliageInfo(UFoliageType* FoliageType, TUniqueObj<FFoliageInfo>&& FoliageInfo)
+{
+	TUniqueObj<FFoliageInfo>& NewFoliageInfo = FoliageInfos.Add(FoliageType, MoveTemp(FoliageInfo));
+	NewFoliageInfo->IFA = this;
+	return NewFoliageInfo;
+}
+
+bool AInstancedFoliageActor::RemoveFoliageInfoAndCopyValue(UFoliageType* FoliageType, TUniqueObj<FFoliageInfo>& OutFoliageInfo)
+{
+	return FoliageInfos.RemoveAndCopyValue(FoliageType, OutFoliageInfo);
+}
 
 UFoliageType* AInstancedFoliageActor::GetLocalFoliageTypeForSource(const UObject* InSource, FFoliageInfo** OutMeshInfo)
 {
@@ -2686,11 +2720,11 @@ void AInstancedFoliageActor::MoveInstancesForMovedOwnedActors(AActor* InActor)
 				if (ActorIndex != INDEX_NONE)
 				{
 					AInstancedFoliageActor* TargetIFA = AInstancedFoliageActor::Get(GetWorld(), true, GetLevel(), InActor->GetActorLocation());
-					FoliageActor->UpdateInstanceFromActor(this, ActorIndex, *Pair.Value);
+					FoliageActor->UpdateInstanceFromActor(ActorIndex, *Pair.Value);
 					if (TargetIFA != this)
 					{
 						// After Moving the Actor doesn't have the same TargetIFA. Reassign.
-						Pair.Value->MoveInstances(this, TargetIFA, { ActorIndex }, true);
+						Pair.Value->MoveInstances(TargetIFA, { ActorIndex }, true);
 					}
 					break;
 				}
@@ -2770,7 +2804,7 @@ void AInstancedFoliageActor::MoveInstancesForMovedComponent(UActorComponent* InC
 			}
 
 			Info.Implementation->EndUpdate();
-			Info.Refresh(this, true, false);
+			Info.Refresh(true, false);
 		}
 	}
 }
@@ -2790,7 +2824,7 @@ void AInstancedFoliageActor::DeleteInstancesForComponent(UActorComponent* InComp
 		const auto* InstanceSet = Info.ComponentHash.Find(BaseId);
 		if (InstanceSet)
 		{
-			Info.RemoveInstances(this, InstanceSet->Array(), true);
+			Info.RemoveInstances(InstanceSet->Array(), true);
 		}
 	}
 }
@@ -2810,7 +2844,7 @@ void AInstancedFoliageActor::DeleteInstancesForComponent(UActorComponent* InComp
 		const auto* InstanceSet = Info->ComponentHash.Find(BaseId);
 		if (InstanceSet)
 		{
-			Info->RemoveInstances(this, InstanceSet->Array(), true);
+			Info->RemoveInstances(InstanceSet->Array(), true);
 		}
 	}
 }
@@ -2842,7 +2876,7 @@ void AInstancedFoliageActor::DeleteInstancesForProceduralFoliageComponent(const
 
 		if (InstancesToRemove.Num())
 		{
-			Info.RemoveInstances(this, InstancesToRemove, InRebuildTree);
+			Info.RemoveInstances(InstancesToRemove, InRebuildTree);
 		}
 	}
 
@@ -2912,13 +2946,13 @@ void AInstancedFoliageActor::MoveInstancesForComponentToLevel(UActorComponent* I
 					// Add the foliage to the new level
 					for (int32 InstanceIndex : *InstanceSet)
 					{
-						TargetMeshInfo->AddInstance(NewIFA, TargetFoliageType, Info.Instances[InstanceIndex], InComponent);
+						TargetMeshInfo->AddInstance(TargetFoliageType, Info.Instances[InstanceIndex], InComponent);
 					}
 
-					TargetMeshInfo->Refresh(NewIFA, true, true);
+					TargetMeshInfo->Refresh(true, true);
 
 					// Remove from old level
-					Info.RemoveInstances(IFA, InstanceSet->Array(), true);
+					Info.RemoveInstances(InstanceSet->Array(), true);
 				}
 			}
 		}
@@ -3028,9 +3062,9 @@ void AInstancedFoliageActor::MoveInstancesToNewComponent(UPrimitiveComponent* In
 				UFoliageType* TargetFoliageType = TargetIFA->AddFoliageType(Pair.Key, &TargetMeshInfo);
 				for (FFoliageInstance* InstanceToMove : IFAPair.Value)
 				{
-					TargetMeshInfo->AddInstance(TargetIFA, TargetFoliageType, *InstanceToMove, InNewComponent);
+					TargetMeshInfo->AddInstance(TargetFoliageType, *InstanceToMove, InNewComponent);
 				}
-				TargetMeshInfo->Refresh(TargetIFA, true, true);
+				TargetMeshInfo->Refresh(true, true);
 			}
 									
 			// Remove old set if empty
@@ -3047,7 +3081,7 @@ void AInstancedFoliageActor::MoveInstancesToNewComponent(UPrimitiveComponent* In
 			// Remove from old IFA
 			if (InstancesToDelete.Num())
 			{
-				Info.RemoveInstances(this, InstancesToDelete, true);
+				Info.RemoveInstances(InstancesToDelete, true);
 			}
 		}
 	}
@@ -3079,7 +3113,7 @@ void AInstancedFoliageActor::MoveInstancesToLevel(ULevel* InTargetLevel, TSet<in
 	}
 
 	AInstancedFoliageActor* TargetIFA = GetInstancedFoliageActorForLevel(InTargetLevel, /*bCreateIfNone*/ true);
-	InCurrentMeshInfo->MoveInstances(this, TargetIFA, InInstanceList, bSelect);
+	InCurrentMeshInfo->MoveInstances(TargetIFA, InInstanceList, bSelect);
 }
 
 void AInstancedFoliageActor::MoveSelectedInstancesToLevel(ULevel* InTargetLevel)
@@ -3167,7 +3201,7 @@ UFoliageType* AInstancedFoliageActor::AddFoliageType(const UFoliageType* InType,
 		if (!ExistingMeshInfo)
 		{
 			Modify();
-			Info = &FoliageInfos.Add(FoliageType).Get();
+			Info = &AddFoliageInfo(FoliageType).Get();
 		}
 		else
 		{
@@ -3182,7 +3216,7 @@ UFoliageType* AInstancedFoliageActor::AddFoliageType(const UFoliageType* InType,
 		{
 			Modify();
 			FoliageType = DuplicateObject<UFoliageType>(InType, this);
-			Info = &FoliageInfos.Add(FoliageType).Get();
+			Info = &AddFoliageInfo(FoliageType).Get();
 		}
 		else
 		{
@@ -3198,11 +3232,11 @@ UFoliageType* AInstancedFoliageActor::AddFoliageType(const UFoliageType* InType,
 		{
 			Modify();
 			FoliageType = DuplicateObject<UFoliageType>(InType, this);
-			Info = &FoliageInfos.Add(FoliageType).Get();
+			Info = &AddFoliageInfo(FoliageType).Get();
 		}
 	}
 
-	if (Info && !Info->Implementation.IsValid())
+	if (Info &&!Info->Implementation.IsValid())
 	{
 		Info->CreateImplementation(FoliageType);
 		check(Info->Implementation.IsValid());
@@ -3213,6 +3247,7 @@ UFoliageType* AInstancedFoliageActor::AddFoliageType(const UFoliageType* InType,
 		*OutInfo = Info;
 	}
 
+	
 	return FoliageType;
 }
 
@@ -3253,7 +3288,7 @@ FFoliageInfo* AInstancedFoliageActor::AddMesh(UFoliageType* InType)
 
 	Modify();
 
-	FFoliageInfo* Info = &*FoliageInfos.Add(InType);
+	FFoliageInfo* Info = &*AddFoliageInfo(InType);
 	if (!Info->Implementation.IsValid())
 	{
 		Info->CreateImplementation(InType);
@@ -3521,7 +3556,7 @@ void AInstancedFoliageActor::PreEditUndo()
 	{
 		FFoliageInfo& Info = *Pair.Value;
 
-		Info.PreEditUndo(this, Pair.Key);
+		Info.PreEditUndo(Pair.Key);
 	}
 }
 
@@ -3541,6 +3576,25 @@ void AInstancedFoliageActor::PostEditUndo()
 	}
 }
 
+void AInstancedFoliageActor::PostDuplicate(bool bDuplicateForPIE)
+{
+	Super::PostDuplicate(bDuplicateForPIE);
+
+	if (!bDuplicateForPIE)
+	{
+		// Fix up in case some code path duplicates an IFA outside of PIE (Duplicate of a Map)
+		for (auto& Pair : FoliageInfos)
+		{
+			FFoliageInfo& Info = *Pair.Value;
+			Info.IFA = this;
+			if (FFoliageImpl* Impl = Info.Implementation.Get())
+			{
+				Impl->Info = &Info;
+			}
+		}
+	}
+}
+
 bool AInstancedFoliageActor::ShouldExport()
 {
 	// We don't support exporting/importing InstancedFoliageActor itself
@@ -3669,10 +3723,10 @@ void AInstancedFoliageActor::MapRebuild()
 		FFoliageInfo& Info = *FindOrAddMesh(Settings);
 		for (FFoliageInstance& Instance : NewInstancePair.Value)
 		{
-			Info.AddInstance(this, Settings, Instance);
+			Info.AddInstance(Settings, Instance);
 		}
 
-		Info.Refresh(this, true, true);
+		Info.Refresh(true, true);
 	}
 }
 
@@ -3914,6 +3968,9 @@ void AInstancedFoliageActor::PostLoad()
 		{
 			// Find the per-mesh info matching the mesh.
 			FFoliageInfo& Info = *Pair.Value;
+			// Make sure to set that before doing anything else (might already been done in Serialize except if the post load as upgraded/modified data)
+			Info.IFA = this;
+			
 			UFoliageType* FoliageType = Pair.Key;
 
 			// Make sure the source data has been PostLoaded as if not it can be considered invalid resulting in a bad HISMC tree
@@ -3940,7 +3997,7 @@ void AInstancedFoliageActor::PostLoad()
 					->AddToken(FTextToken::Create(FText::Format(LOCTEXT("MapCheck_Message_FoliageMissingComponent", "Foliage in this map is missing a component for static mesh {MeshName}. This has been repaired."), Arguments)))
 					->AddToken(FMapErrorToken::Create(FMapErrors::FoliageMissingClusterComponent));
 
-				Info.ReallocateClusters(this, Pair.Key);
+				Info.ReallocateClusters(Pair.Key);
 			}
 
 			// Update the hash.
@@ -3961,7 +4018,7 @@ void AInstancedFoliageActor::PostLoad()
 			// Convert to Hierarchical foliage
 			if (GetLinkerCustomVersion(FFoliageCustomVersion::GUID) < FFoliageCustomVersion::FoliageUsingHierarchicalISMC)
 			{
-				Info.ReallocateClusters(this, Pair.Key);
+				Info.ReallocateClusters(Pair.Key);
 			}
 
 			if (GetLinkerCustomVersion(FFoliageCustomVersion::GUID) < FFoliageCustomVersion::HierarchicalISMCNonTransactional)
@@ -3994,7 +4051,7 @@ void AInstancedFoliageActor::PostLoad()
 				{
 					FFoliageStaticMesh* FoliageStaticMesh = StaticCast<FFoliageStaticMesh*>(Info.Implementation.Get());
 					UFoliageType_InstancedStaticMesh* FoliageType_InstancedStaticMesh = Cast<UFoliageType_InstancedStaticMesh>(FoliageType);
-					FoliageStaticMesh->CheckComponentClass(this, FoliageType_InstancedStaticMesh, Info.Instances, Info.SelectedIndices);
+					FoliageStaticMesh->CheckComponentClass(FoliageType_InstancedStaticMesh, Info.Instances, Info.SelectedIndices);
 				}
 			}
 
@@ -4214,11 +4271,11 @@ void AInstancedFoliageActor::RepairDuplicateIFA(AInstancedFoliageActor* Duplicat
 		{
 			if ((Instance.Flags & FOLIAGE_InstanceDeleted) == 0)
 			{
-				TargetMeshInfo->AddInstance(this, TargetFoliageType, Instance);
+				TargetMeshInfo->AddInstance(TargetFoliageType, Instance);
 			}
 		}
 
-		TargetMeshInfo->Refresh(this, true, true);
+		TargetMeshInfo->Refresh(true, true);
 	}
 
 	GetWorld()->DestroyActor(DuplicateIFA);
@@ -4230,7 +4287,7 @@ void AInstancedFoliageActor::NotifyFoliageTypeChanged(UFoliageType* FoliageType,
 
 	if (TypeInfo)
 	{
-		TypeInfo->NotifyFoliageTypeChanged(this, FoliageType, bSourceChanged);	
+		TypeInfo->NotifyFoliageTypeChanged(FoliageType, bSourceChanged);	
 
 		if (bSourceChanged)
 		{
@@ -4281,7 +4338,7 @@ void AInstancedFoliageActor::AddInstances(UObject* WorldContextObject, UFoliageT
 			FFoliageInfo* TypeInfo = nullptr;
 			if (UFoliageType* FoliageType = Pair.Key->AddFoliageType(InFoliageType, &TypeInfo))
 			{
-				TypeInfo->AddInstances(Pair.Key, FoliageType, Pair.Value);
+				TypeInfo->AddInstances(FoliageType, Pair.Value);
 			}
 		}
 	}
@@ -4307,7 +4364,7 @@ void AInstancedFoliageActor::NotifyFoliageTypeWillChange(UFoliageType* FoliageTy
 	// Change bounds delegate bindings
 	if (TypeInfo)
 	{
-		TypeInfo->NotifyFoliageTypeWillChange(this, FoliageType);
+		TypeInfo->NotifyFoliageTypeWillChange(FoliageType);
 	}
 }
 
@@ -4429,7 +4486,7 @@ void AInstancedFoliageActor::OnLevelActorDeleted(AActor* InActor)
 						
 						if(InstancesToRemove.Num() > 0)
 						{
-							Info.RemoveInstances(this, InstancesToRemove, false);
+							Info.RemoveInstances(InstancesToRemove, false);
 							if (InstanceCountChanged.IsBound())
 							{
 								InstanceCountChanged.Broadcast(FoliageType);
@@ -4468,7 +4525,7 @@ void AInstancedFoliageActor::DetectFoliageTypeChangeAndUpdate()
 				{
 					FFoliageStaticMesh* FoliageStaticMesh = StaticCast<FFoliageStaticMesh*>(Info.Implementation.Get());
 					UFoliageType_InstancedStaticMesh* FoliageType_InstancedStaticMesh = Cast<UFoliageType_InstancedStaticMesh>(FoliageType);
-					FoliageStaticMesh->CheckComponentClass(this, FoliageType_InstancedStaticMesh, Info.Instances, Info.SelectedIndices);
+					FoliageStaticMesh->CheckComponentClass(FoliageType_InstancedStaticMesh, Info.Instances, Info.SelectedIndices);
 					FoliageStaticMesh->UpdateComponentSettings(FoliageType_InstancedStaticMesh);
 				}
 				// Respawn foliage actors
@@ -4480,7 +4537,7 @@ void AInstancedFoliageActor::DetectFoliageTypeChangeAndUpdate()
 					{
 						FFoliageActor* FoliageActor = StaticCast<FFoliageActor*>(Info.Implementation.Get());
 						const bool bPostLoad = true;
-						FoliageActor->Reapply(this, FoliageType, Info.Instances, bPostLoad);
+						FoliageActor->Reapply(FoliageType, Info.Instances, bPostLoad);
 					}
 				}
 			}
@@ -4563,7 +4620,7 @@ void AInstancedFoliageActor::CleanupDeletedFoliageType()
 
 			if (InstancesToRemove.Num())
 			{
-				Info.RemoveInstances(this, InstancesToRemove, true);
+				Info.RemoveInstances(InstancesToRemove, true);
 			}
 
 			It.RemoveCurrent();
diff --git a/Engine/Source/Runtime/Foliage/Public/InstancedFoliage.h b/Engine/Source/Runtime/Foliage/Public/InstancedFoliage.h
--- a/Engine/Source/Runtime/Foliage/Public/InstancedFoliage.h
+++ b/Engine/Source/Runtime/Foliage/Public/InstancedFoliage.h
@@ -168,32 +168,43 @@ enum class EFoliageImplType : uint8
 	Actor = 2
 };
 
+struct FFoliageInfo;
+
 ///
 /// FFoliageInfoImpl
 ///
 struct FFoliageImpl
 {
-	FFoliageImpl() {}
+	FFoliageImpl(FFoliageInfo* InInfo) 
+#if WITH_EDITORONLY_DATA
+		: Info(InInfo)
+#endif
+	{}
 	virtual ~FFoliageImpl() {}
 
 	virtual void AddReferencedObjects(UObject* InThis, FReferenceCollector& Collector) {}
 	virtual void Serialize(FArchive& Ar) = 0;
 
+#if WITH_EDITORONLY_DATA
+	// Not serialized but FFoliageInfo will make sure it stays valid (mostly Undo/Redo)
+	FFoliageInfo* Info;
+#endif
+
 #if WITH_EDITOR
 	virtual bool IsInitialized() const = 0;
-	virtual void Initialize(AInstancedFoliageActor* IFA, const UFoliageType* FoliageType) = 0;
+	virtual void Initialize(const UFoliageType* FoliageType) = 0;
 	virtual void Uninitialize() = 0;
 	virtual int32 GetInstanceCount() const = 0;
-	virtual void PreAddInstances(AInstancedFoliageActor* IFA, const UFoliageType* FoliageType, int32 Count) = 0;
-	virtual void AddInstance(AInstancedFoliageActor* IFA, const FFoliageInstance& NewInstance) = 0;
+	virtual void PreAddInstances(const UFoliageType* FoliageType, int32 Count) = 0;
+	virtual void AddInstance(const FFoliageInstance& NewInstance) = 0;
 	virtual void RemoveInstance(int32 InstanceIndex) = 0;
 	virtual void MoveInstance(int32 InstanceIndex, UObject*& OutInstanceImplementation) { RemoveInstance(InstanceIndex); }
-	virtual void AddExistingInstance(AInstancedFoliageActor* IFA, const FFoliageInstance& ExistingInstance, UObject* InstanceImplementation) { AddInstance(IFA, ExistingInstance); }
+	virtual void AddExistingInstance(const FFoliageInstance& ExistingInstance, UObject* InstanceImplementation) { AddInstance(ExistingInstance); }
 	virtual void SetInstanceWorldTransform(int32 InstanceIndex, const FTransform& Transform, bool bTeleport) = 0;
 	virtual FTransform GetInstanceWorldTransform(int32 InstanceIndex) const = 0;
 	virtual void PostUpdateInstances() {}
-	virtual void PreMoveInstances(AInstancedFoliageActor* InIFA, const TArray<int32>& InInstancesMoved) {}
-	virtual void PostMoveInstances(AInstancedFoliageActor* InIFA, const TArray<int32>& InInstancesMoved, bool bFinished) {}
+	virtual void PreMoveInstances(const TArray<int32>& InInstancesMoved) {}
+	virtual void PostMoveInstances(const TArray<int32>& InInstancesMoved, bool bFinished) {}
 	virtual bool IsOwnedComponent(const UPrimitiveComponent* Component) const = 0;
 	virtual int32 FindIndex(const UPrimitiveComponent* HitComponent) const { return INDEX_NONE; }
 
@@ -206,15 +217,18 @@ struct FFoliageImpl
 
 	virtual void BeginUpdate() {}
 	virtual void EndUpdate() {}
-	virtual void Refresh(AInstancedFoliageActor* IFA, const TArray<FFoliageInstance>& Instances, bool Async, bool Force) {}
+	virtual void Refresh(const TArray<FFoliageInstance>& Instances, bool Async, bool Force) {}
 	virtual void OnHiddenEditorViewMaskChanged(uint64 InHiddenEditorViews) = 0;
-	virtual void PreEditUndo(AInstancedFoliageActor* IFA, UFoliageType* FoliageType) {}
-	virtual void PostEditUndo(AInstancedFoliageActor* IFA, UFoliageType* FoliageType, const TArray<FFoliageInstance>& Instances, const TSet<int32>& SelectedIndices) = 0;
-	virtual void NotifyFoliageTypeWillChange(AInstancedFoliageActor* IFA, UFoliageType* FoliageType) {}
-	virtual void NotifyFoliageTypeChanged(AInstancedFoliageActor* IFA, UFoliageType* FoliageType, const TArray<FFoliageInstance>& Instances, const TSet<int32>& SelectedIndices, bool bSourceChanged) = 0;
+	virtual void PreEditUndo(UFoliageType* FoliageType) {}
+	virtual void PostEditUndo(FFoliageInfo* InInfo, UFoliageType* FoliageType, const TArray<FFoliageInstance>& Instances, const TSet<int32>& SelectedIndices) { Info = InInfo; }
+	virtual void NotifyFoliageTypeWillChange(UFoliageType* FoliageType) {}
+	virtual void NotifyFoliageTypeChanged(UFoliageType* FoliageType, const TArray<FFoliageInstance>& Instances, const TSet<int32>& SelectedIndices, bool bSourceChanged) = 0;
 	virtual void EnterEditMode() {}
 	virtual void ExitEditMode() {}
 	virtual bool ShouldAttachToBaseComponent() const { return true; }
+
+	AInstancedFoliageActor* GetIFA() const;
+	FFoliageInfo* GetInfo() const { check(Info); return Info; }
 #endif
 
 	virtual int32 GetOverlappingSphereCount(const FSphere& Sphere) const { return 0; }
@@ -232,6 +246,9 @@ struct FFoliageInfo
 	TUniquePtr<FFoliageImpl> Implementation;
 
 #if WITH_EDITORONLY_DATA
+	// Owning IFA
+	AInstancedFoliageActor* IFA;
+
 	// Allows us to detect if FoliageType was updated while this level wasn't loaded
 	FGuid FoliageTypeUpdateGuid;
 
@@ -263,7 +280,7 @@ struct FFoliageInfo
 	void AddReferencedObjects(UObject* InThis, FReferenceCollector& Collector);
 		
 	FOLIAGE_API void CreateImplementation(EFoliageImplType InType);
-	FOLIAGE_API void Initialize(AInstancedFoliageActor* InIFA, const UFoliageType* FoliageType);
+	FOLIAGE_API void Initialize(const UFoliageType* FoliageType);
 	FOLIAGE_API void Uninitialize();
 	FOLIAGE_API bool IsInitialized() const;
 		
@@ -275,41 +292,41 @@ struct FFoliageInfo
 #if WITH_EDITOR
 	
 	FOLIAGE_API void CreateImplementation(const UFoliageType* FoliageType);
-	FOLIAGE_API void NotifyFoliageTypeWillChange(AInstancedFoliageActor* IFA, UFoliageType* FoliageType);
-	FOLIAGE_API void NotifyFoliageTypeChanged(AInstancedFoliageActor* IFA, UFoliageType* FoliageType, bool bSourceChanged);
+	FOLIAGE_API void NotifyFoliageTypeWillChange(UFoliageType* FoliageType);
+	FOLIAGE_API void NotifyFoliageTypeChanged(UFoliageType* FoliageType, bool bSourceChanged);
 	
 	FOLIAGE_API void ClearSelection();
 
 	FOLIAGE_API void SetInstanceWorldTransform(int32 InstanceIndex, const FTransform& Transform, bool bTeleport);
 
 	FOLIAGE_API void SetRandomSeed(int32 seed);
-	FOLIAGE_API void AddInstance(AInstancedFoliageActor* InIFA, const UFoliageType* InSettings, const FFoliageInstance& InNewInstance);
-	FOLIAGE_API void AddInstance(AInstancedFoliageActor* InIFA, const UFoliageType* InSettings, const FFoliageInstance& InNewInstance, UActorComponent* InBaseComponent);
-	FOLIAGE_API void AddInstances(AInstancedFoliageActor* InIFA, const UFoliageType* InSettings, const TArray<const FFoliageInstance*>& InNewInstances);
-	FOLIAGE_API void ReserveAdditionalInstances(AInstancedFoliageActor* InIFA, const UFoliageType* InSettings, uint32 ReserveNum);
+	FOLIAGE_API void AddInstance(const UFoliageType* InSettings, const FFoliageInstance& InNewInstance);
+	FOLIAGE_API void AddInstance(const UFoliageType* InSettings, const FFoliageInstance& InNewInstance, UActorComponent* InBaseComponent);
+	FOLIAGE_API void AddInstances(const UFoliageType* InSettings, const TArray<const FFoliageInstance*>& InNewInstances);
+	FOLIAGE_API void ReserveAdditionalInstances(const UFoliageType* InSettings, uint32 ReserveNum);
 
-	FOLIAGE_API void RemoveInstances(AInstancedFoliageActor* InIFA, const TArray<int32>& InInstancesToRemove, bool RebuildFoliageTree);
+	FOLIAGE_API void RemoveInstances(const TArray<int32>& InInstancesToRemove, bool RebuildFoliageTree);
 
-	FOLIAGE_API void MoveInstances(AInstancedFoliageActor* InFromIFA, AInstancedFoliageActor* InToIFA, const TSet<int32>& InInstancesToMove, bool bKeepSelection);
+	FOLIAGE_API void MoveInstances(AInstancedFoliageActor* InToIFA, const TSet<int32>& InInstancesToMove, bool bKeepSelection);
 
 	// Apply changes in the FoliageType to the component
-	FOLIAGE_API void PreMoveInstances(AInstancedFoliageActor* InIFA, const TArray<int32>& InInstancesToMove);
-	FOLIAGE_API void PostMoveInstances(AInstancedFoliageActor* InIFA, const TArray<int32>& InInstancesMoved, bool bFinished = false);
-	FOLIAGE_API void PostUpdateInstances(AInstancedFoliageActor* InIFA, const TArray<int32>& InInstancesUpdated, bool bReAddToHash = false, bool InUpdateSelection = false);
-	FOLIAGE_API void DuplicateInstances(AInstancedFoliageActor* InIFA, UFoliageType* InSettings, const TArray<int32>& InInstancesToDuplicate);
-	FOLIAGE_API void GetInstancesInsideBounds(const FBox& Box, TArray<int32>& OutInstances);
-	FOLIAGE_API void GetInstancesInsideSphere(const FSphere& Sphere, TArray<int32>& OutInstances);
-	FOLIAGE_API void GetInstanceAtLocation(const FVector& Location, int32& OutInstance, bool& bOutSucess);
-	FOLIAGE_API bool CheckForOverlappingSphere(const FSphere& Sphere);
-	FOLIAGE_API bool CheckForOverlappingInstanceExcluding(int32 TestInstanceIdx, float Radius, TSet<int32>& ExcludeInstances);
+	FOLIAGE_API void PreMoveInstances(const TArray<int32>& InInstancesToMove);
+	FOLIAGE_API void PostMoveInstances(const TArray<int32>& InInstancesMoved, bool bFinished = false);
+	FOLIAGE_API void PostUpdateInstances(const TArray<int32>& InInstancesUpdated, bool bReAddToHash = false, bool InUpdateSelection = false);
+	FOLIAGE_API void DuplicateInstances(UFoliageType* InSettings, const TArray<int32>& InInstancesToDuplicate);
+	FOLIAGE_API void GetInstancesInsideBounds(const FBox& Box, TArray<int32>& OutInstances) const;
+	FOLIAGE_API void GetInstancesInsideSphere(const FSphere& Sphere, TArray<int32>& OutInstances) const;
+	FOLIAGE_API void GetInstanceAtLocation(const FVector& Location, int32& OutInstance, bool& bOutSucess) const;
+	FOLIAGE_API bool CheckForOverlappingSphere(const FSphere& Sphere) const;
+	FOLIAGE_API bool CheckForOverlappingInstanceExcluding(int32 TestInstanceIdx, float Radius, TSet<int32>& ExcludeInstances) const;
 	FOLIAGE_API TArray<int32> GetInstancesOverlappingBox(const FBox& Box) const;
 
 	// Destroy existing clusters and reassign all instances to new clusters
-	FOLIAGE_API void ReallocateClusters(AInstancedFoliageActor* InIFA, UFoliageType* InSettings);
+	FOLIAGE_API void ReallocateClusters(UFoliageType* InSettings);
 
-	FOLIAGE_API void SelectInstances(AInstancedFoliageActor* InIFA, bool bSelect, TArray<int32>& Instances);
+	FOLIAGE_API void SelectInstances(bool bSelect, TArray<int32>& Instances);
 
-	FOLIAGE_API void SelectInstances(AInstancedFoliageActor* InIFA, bool bSelect);
+	FOLIAGE_API void SelectInstances(bool bSelect);
 
 	FOLIAGE_API FBox GetSelectionBoundingBox() const;
 
@@ -323,15 +340,15 @@ struct FFoliageInfo
 	// For debugging. Validate state after editing.
 	void CheckValid();
 		
-	FOLIAGE_API void Refresh(AInstancedFoliageActor* IFA, bool Async, bool Force);
+	FOLIAGE_API void Refresh(bool Async, bool Force);
 	FOLIAGE_API void OnHiddenEditorViewMaskChanged(uint64 InHiddenEditorViews);
-	FOLIAGE_API void PostEditUndo(AInstancedFoliageActor* IFA, UFoliageType* FoliageType);
-	FOLIAGE_API void PreEditUndo(AInstancedFoliageActor* IFA, UFoliageType* FoliageType);
+	FOLIAGE_API void PostEditUndo(AInstancedFoliageActor* InIFA, UFoliageType* FoliageType);
+	FOLIAGE_API void PreEditUndo(UFoliageType* FoliageType);
 	FOLIAGE_API void EnterEditMode();
 	FOLIAGE_API void ExitEditMode();
 
 	FOLIAGE_API void RemoveBaseComponentOnInstances();
-	FOLIAGE_API void IncludeActor(AInstancedFoliageActor* IFA, const UFoliageType* FoliageType, AActor* InActor);
+	FOLIAGE_API void IncludeActor(const UFoliageType* FoliageType, AActor* InActor);
 	FOLIAGE_API void ExcludeActors();
 
 	FOLIAGE_API FBox GetApproximatedInstanceBounds() const;
@@ -345,11 +362,11 @@ struct FFoliageInfo
 
 private:
 	using FAddImplementationFunc = TFunctionRef<void(FFoliageImpl*, AInstancedFoliageActor*, const FFoliageInstance&)>;
-	void AddInstancesImpl(AInstancedFoliageActor* InIFA, const UFoliageType* InSettings, const TArray<const FFoliageInstance*>& InNewInstances, FFoliageInfo::FAddImplementationFunc ImplementationFunc);
-	void AddInstanceImpl(AInstancedFoliageActor* InIFA, const FFoliageInstance& InNewInstance, FAddImplementationFunc ImplementationFunc);
+	void AddInstancesImpl(const UFoliageType* InSettings, const TArray<const FFoliageInstance*>& InNewInstances, FFoliageInfo::FAddImplementationFunc ImplementationFunc);
+	void AddInstanceImpl(const FFoliageInstance& InNewInstance, FAddImplementationFunc ImplementationFunc);
 
 	using FRemoveImplementationFunc = TFunctionRef<void(FFoliageImpl*, int32)>;
-	void RemoveInstancesImpl(AInstancedFoliageActor* InIFA, const TArray<int32>& InInstancesToRemove, bool RebuildFoliageTree, FRemoveImplementationFunc ImplementationFunc);
+	void RemoveInstancesImpl(const TArray<int32>& InInstancesToRemove, bool RebuildFoliageTree, FRemoveImplementationFunc ImplementationFunc);
 };
 
 
diff --git a/Engine/Source/Runtime/Foliage/Public/InstancedFoliageActor.h b/Engine/Source/Runtime/Foliage/Public/InstancedFoliageActor.h
--- a/Engine/Source/Runtime/Foliage/Public/InstancedFoliageActor.h
+++ b/Engine/Source/Runtime/Foliage/Public/InstancedFoliageActor.h
@@ -33,9 +33,18 @@ public:
 	UActorComponent* GetBaseComponentFromBaseId(const FFoliageInstanceBaseId& BaseId) const;
 #endif// WITH_EDITORONLY_DATA
 
+private:
+	friend FFoliageInstanceBaseCache;
+	
 	TMap<UFoliageType*, TUniqueObj<FFoliageInfo>> FoliageInfos;
 
 public:
+	FOLIAGE_API void ForEachFoliageInfo(TFunctionRef<bool(UFoliageType* FoliageType, FFoliageInfo& FoliageInfo)> InOperation);
+	FOLIAGE_API const TMap<UFoliageType*, TUniqueObj<FFoliageInfo>>& GetFoliageInfos() const { return FoliageInfos; }
+	FOLIAGE_API TUniqueObj<FFoliageInfo>& AddFoliageInfo(UFoliageType* FoliageType);
+	FOLIAGE_API TUniqueObj<FFoliageInfo>& AddFoliageInfo(UFoliageType* FoliageType, TUniqueObj<FFoliageInfo>&& FoliageInfo);
+	FOLIAGE_API bool RemoveFoliageInfoAndCopyValue(UFoliageType* FoliageType, TUniqueObj<FFoliageInfo>& OutFoliageInfo);
+
 	//~ Begin UObject Interface.
 	virtual void Serialize(FArchive& Ar) override;
 	virtual void PostLoad() override;
@@ -132,6 +141,7 @@ public:
 
 	virtual void PreEditUndo() override;
 	virtual void PostEditUndo() override;
+	virtual void PostDuplicate(bool bDuplicateForPIE) override;
 	virtual bool ShouldExport() override;
 	virtual bool ShouldImport(FString* ActorPropString, bool IsMovingLevel) override;
 
diff --git a/Engine/Source/Runtime/Landscape/Private/LandscapeCollision.cpp b/Engine/Source/Runtime/Landscape/Private/LandscapeCollision.cpp
--- a/Engine/Source/Runtime/Landscape/Private/LandscapeCollision.cpp
+++ b/Engine/Source/Runtime/Landscape/Private/LandscapeCollision.cpp
@@ -2022,12 +2022,8 @@ void ULandscapeHeightfieldCollisionComponent::SnapFoliageInstances(const FBox& I
 			continue;
 		}
 			
-		for (auto& Pair : IFA->FoliageInfos)
+		IFA->ForEachFoliageInfo([this, IFA, BaseId, &InInstanceBox](UFoliageType* Settings, FFoliageInfo& MeshInfo)
 		{
-			// Find the per-mesh info matching the mesh.
-			UFoliageType* Settings = Pair.Key;
-			FFoliageInfo& MeshInfo = *Pair.Value;
-			
 			const auto* InstanceSet = MeshInfo.ComponentHash.Find(BaseId);
 			if (InstanceSet)
 			{
@@ -2114,14 +2110,15 @@ void ULandscapeHeightfieldCollisionComponent::SnapFoliageInstances(const FBox& I
 				}
 
 				// Remove any unused instances
-				MeshInfo.RemoveInstances(IFA, InstancesToRemove, true);
+				MeshInfo.RemoveInstances(InstancesToRemove, true);
 
 				for (UHierarchicalInstancedStaticMeshComponent* FoliageComp : AffectedFoliageComponents)
 				{
 					FoliageComp->InvalidateLightingCache();
 				}
 			}
-		}
+			return true; // continue iterating
+		});
 	}
 }
 #endif // WITH_EDITORONLY_DATA
