Fix Non-editor Build

diff --git a/Engine/Source/Runtime/Foliage/Private/InstancedFoliage.cpp b/Engine/Source/Runtime/Foliage/Private/InstancedFoliage.cpp
--- a/Engine/Source/Runtime/Foliage/Private/InstancedFoliage.cpp
+++ b/Engine/Source/Runtime/Foliage/Private/InstancedFoliage.cpp
@@ -2550,14 +2550,18 @@ void AInstancedFoliageActor::ForEachFoliageInfo(TFunctionRef<bool(UFoliageType*
 TUniqueObj<FFoliageInfo>& AInstancedFoliageActor::AddFoliageInfo(UFoliageType* FoliageType)
 {
 	TUniqueObj<FFoliageInfo>& NewFoliageInfo = FoliageInfos.Add(FoliageType);
+#if WITH_EDITORONLY_DATA
 	NewFoliageInfo->IFA = this;
+#endif
 	return NewFoliageInfo;
 }
 
 TUniqueObj<FFoliageInfo>& AInstancedFoliageActor::AddFoliageInfo(UFoliageType* FoliageType, TUniqueObj<FFoliageInfo>&& FoliageInfo)
 {
 	TUniqueObj<FFoliageInfo>& NewFoliageInfo = FoliageInfos.Add(FoliageType, MoveTemp(FoliageInfo));
+#if WITH_EDITORONLY_DATA
 	NewFoliageInfo->IFA = this;
+#endif
 	return NewFoliageInfo;
 }
 
diff --git a/Engine/Source/Runtime/Foliage/Public/InstancedFoliageActor.h b/Engine/Source/Runtime/Foliage/Public/InstancedFoliageActor.h
--- a/Engine/Source/Runtime/Foliage/Public/InstancedFoliageActor.h
+++ b/Engine/Source/Runtime/Foliage/Public/InstancedFoliageActor.h
@@ -34,7 +34,7 @@ public:
 #endif// WITH_EDITORONLY_DATA
 
 private:
-	friend FFoliageInstanceBaseCache;
+	friend class FFoliageInstanceBaseCache;
 	
 	TMap<UFoliageType*, TUniqueObj<FFoliageInfo>> FoliageInfos;
 
