Build fix...for real.

diff --git a/Engine/Source/Runtime/Foliage/Public/InstancedFoliageActor.h b/Engine/Source/Runtime/Foliage/Public/InstancedFoliageActor.h
--- a/Engine/Source/Runtime/Foliage/Public/InstancedFoliageActor.h
+++ b/Engine/Source/Runtime/Foliage/Public/InstancedFoliageActor.h
@@ -34,7 +34,7 @@ public:
 #endif// WITH_EDITORONLY_DATA
 
 private:
-	friend class FFoliageInstanceBaseCache;
+	friend struct FFoliageInstanceBaseCache;
 	
 	TMap<UFoliageType*, TUniqueObj<FFoliageInfo>> FoliageInfos;
 
