New ISMPartitionActor (now parent class of AInstancedFoliageActor)
- This actor allows registering ISM Clients that can add/remove their own instances represented by N SM/Transforms. The actor will group in the fewest amount of ISM components.

diff --git a/Engine/Source/Runtime/Engine/Classes/Components/HierarchicalInstancedStaticMeshComponent.h b/Engine/Source/Runtime/Engine/Classes/Components/HierarchicalInstancedStaticMeshComponent.h
--- a/Engine/Source/Runtime/Engine/Classes/Components/HierarchicalInstancedStaticMeshComponent.h
+++ b/Engine/Source/Runtime/Engine/Classes/Components/HierarchicalInstancedStaticMeshComponent.h
@@ -233,6 +233,7 @@ public:
 	virtual TArray<int32> GetInstancesOverlappingSphere(const FVector& Center, float Radius, bool bSphereInWorldSpace = true) const override;
 	virtual TArray<int32> GetInstancesOverlappingBox(const FBox& Box, bool bBoxInWorldSpace = true) const override;
 	virtual void PreAllocateInstancesMemory(int32 AddedInstanceCount) override;
+	virtual bool SupportsRemoveSwap() const override { return true; }
 
 	/** Removes all the instances with indices specified in the InstancesToRemove array. Returns true on success. */
 	UFUNCTION(BlueprintCallable, Category = "Components|InstancedStaticMesh")
diff --git a/Engine/Source/Runtime/Engine/Classes/Components/InstancedStaticMeshComponent.h b/Engine/Source/Runtime/Engine/Classes/Components/InstancedStaticMeshComponent.h
--- a/Engine/Source/Runtime/Engine/Classes/Components/InstancedStaticMeshComponent.h
+++ b/Engine/Source/Runtime/Engine/Classes/Components/InstancedStaticMeshComponent.h
@@ -277,6 +277,7 @@ class ENGINE_API UInstancedStaticMeshComponent : public UStaticMeshComponent
 
 	virtual void PostLoad() override;
 	virtual void OnRegister() override;
+	virtual bool SupportsRemoveSwap() const { return false; }
 
 public:
 	/** Render data will be initialized on PostLoad or on demand. Released on the rendering thread. */
diff --git a/Engine/Source/Runtime/Engine/Private/ISMPartition/ISMComponentData.cpp b/Engine/Source/Runtime/Engine/Private/ISMPartition/ISMComponentData.cpp
new file mode 100644
--- /dev/null
+++ b/Engine/Source/Runtime/Engine/Private/ISMPartition/ISMComponentData.cpp
@@ -0,0 +1,49 @@
+// Copyright Epic Games, Inc. All Rights Reserved.
+
+#include "ISMPartition/ISMComponentData.h"
+#include "Components/InstancedStaticMeshComponent.h"
+#include "Components/HierarchicalInstancedStaticMeshComponent.h"
+
+FISMComponentData::FISMComponentData()
+#if WITH_EDITORONLY_DATA
+	: Component(nullptr)
+	, bInvalidateLightingCache(false)
+	, bAutoRebuildTreeOnInstanceChanges(false)
+	, bWasModifyCalled(false)
+#endif
+{
+
+}
+
+#if WITH_EDITOR
+void FISMComponentData::RegisterDelegates()
+{
+	if (Component && Component->IsA<UHierarchicalInstancedStaticMeshComponent>())
+	{
+		if (Component->GetStaticMesh())
+		{
+			Component->GetStaticMesh()->GetOnExtendedBoundsChanged().AddRaw(this, &FISMComponentData::HandleComponentMeshBoundsChanged);
+		}
+	}
+}
+
+void FISMComponentData::UnregisterDelegates()
+{
+	if (Component && Component->IsA<UHierarchicalInstancedStaticMeshComponent>())
+	{
+		if (Component->GetStaticMesh())
+		{
+			Component->GetStaticMesh()->GetOnExtendedBoundsChanged().RemoveAll(this);
+		}
+	}
+}
+
+void FISMComponentData::HandleComponentMeshBoundsChanged(const FBoxSphereBounds& NewBounds)
+{
+	check(Component);
+	if (UHierarchicalInstancedStaticMeshComponent* HISMComponent = Cast<UHierarchicalInstancedStaticMeshComponent>(Component))
+	{
+		HISMComponent->BuildTreeIfOutdated(true, false);
+	}
+}
+#endif
diff --git a/Engine/Source/Runtime/Engine/Private/ISMPartition/ISMComponentDescriptor.cpp b/Engine/Source/Runtime/Engine/Private/ISMPartition/ISMComponentDescriptor.cpp
new file mode 100644
--- /dev/null
+++ b/Engine/Source/Runtime/Engine/Private/ISMPartition/ISMComponentDescriptor.cpp
@@ -0,0 +1,169 @@
+#include "ISMPartition/ISMComponentDescriptor.h"
+
+#if WITH_EDITOR
+
+#include "Serialization/ArchiveObjectCrc32.h"
+#include "Components/InstancedStaticMeshComponent.h"
+#include "Components/HierarchicalInstancedStaticMeshComponent.h"
+
+#endif
+
+#if WITH_EDITOR
+
+FISMComponentDescriptor FISMComponentDescriptor::CreateFrom(const TSubclassOf<UStaticMeshComponent>& From)
+{
+	FISMComponentDescriptor ComponentDescriptor;
+
+	ComponentDescriptor.InitFrom(From->GetDefaultObject<UStaticMeshComponent>());
+	ComponentDescriptor.ComputeHash();
+
+	return ComponentDescriptor;
+}
+
+void FISMComponentDescriptor::InitFrom(UStaticMeshComponent* Template)
+{
+	bEnableDiscardOnLoad = false;
+	ComponentClass = Template->GetClass();
+	StaticMesh = Template->GetStaticMesh();
+	OverrideMaterials = Template->OverrideMaterials;
+	Mobility = Template->Mobility;
+	RuntimeVirtualTextures = Template->RuntimeVirtualTextures;
+	VirtualTextureRenderPassType = Template->VirtualTextureRenderPassType;
+	LightmapType = Template->LightmapType;
+	LightingChannels = Template->LightingChannels;
+	bHasCustomNavigableGeometry = Template->bHasCustomNavigableGeometry;
+	CustomDepthStencilWriteMask = Template->CustomDepthStencilWriteMask;
+	BodyInstance.CopyBodyInstancePropertiesFrom(&Template->BodyInstance);
+	VirtualTextureCullMips = Template->VirtualTextureCullMips;
+	TranslucencySortPriority = Template->TranslucencySortPriority;
+	OverriddenLightMapRes = Template->OverriddenLightMapRes;
+	CustomDepthStencilValue = Template->CustomDepthStencilValue;
+	bCastShadow = Template->CastShadow;
+	bCastStaticShadow = Template->bCastStaticShadow;
+	bCastDynamicShadow = Template->bCastDynamicShadow;
+	bCastContactShadow = Template->bCastContactShadow;
+	bCastShadowAsTwoSided = Template->bCastShadowAsTwoSided;
+	bAffectDynamicIndirectLighting = Template->bAffectDynamicIndirectLighting;
+	bAffectDistanceFieldLighting = Template->bAffectDistanceFieldLighting;
+	bReceivesDecals = Template->bReceivesDecals;
+	bOverrideLightMapRes = Template->bOverrideLightMapRes;
+	bUseAsOccluder = Template->bUseAsOccluder;
+	bRenderCustomDepth = Template->bRenderCustomDepth;
+	bIncludeInHLOD = Template->bEnableAutoLODGeneration;
+
+	if (UInstancedStaticMeshComponent* ISMTemplate = Cast<UInstancedStaticMeshComponent>(Template))
+	{
+		InstanceStartCullDistance = ISMTemplate->InstanceStartCullDistance;
+		InstanceEndCullDistance = ISMTemplate->InstanceEndCullDistance;
+
+		// HISM Specific
+		if (UHierarchicalInstancedStaticMeshComponent* HISMTemplate = Cast<UHierarchicalInstancedStaticMeshComponent>(Template))
+		{
+			bEnableDensityScaling = HISMTemplate->bEnableDensityScaling;
+		}
+	}
+}
+
+bool FISMComponentDescriptor::operator!=(const FISMComponentDescriptor& Other) const
+{
+	return !(*this == Other);
+}
+
+bool FISMComponentDescriptor::operator==(const FISMComponentDescriptor& Other) const
+{
+	return Hash == Other.Hash && // Check hash first, other checks are in case of Hash collision
+	ComponentClass == Other.ComponentClass &&
+	StaticMesh == Other.StaticMesh &&
+	OverrideMaterials == Other.OverrideMaterials &&
+	Mobility == Other.Mobility &&
+	RuntimeVirtualTextures == Other.RuntimeVirtualTextures &&
+	VirtualTextureRenderPassType == Other.VirtualTextureRenderPassType &&
+	LightmapType == Other.LightmapType &&
+	GetLightingChannelMaskForStruct(LightingChannels) == GetLightingChannelMaskForStruct(Other.LightingChannels) &&
+	bHasCustomNavigableGeometry == Other.bHasCustomNavigableGeometry &&
+	CustomDepthStencilWriteMask == Other.CustomDepthStencilWriteMask &&
+	InstanceStartCullDistance == Other.InstanceStartCullDistance &&
+	InstanceEndCullDistance == Other.InstanceEndCullDistance &&
+	VirtualTextureCullMips == Other.VirtualTextureCullMips &&
+	TranslucencySortPriority == Other.TranslucencySortPriority &&
+	OverriddenLightMapRes == Other.OverriddenLightMapRes &&
+	CustomDepthStencilValue == Other.CustomDepthStencilValue &&
+	bCastShadow == Other.bCastShadow &&
+	bCastStaticShadow == Other.bCastStaticShadow &&
+	bCastDynamicShadow == Other.bCastDynamicShadow &&
+	bCastContactShadow == Other.bCastContactShadow &&
+	bCastShadowAsTwoSided == Other.bCastShadowAsTwoSided &&
+	bAffectDynamicIndirectLighting == Other.bAffectDynamicIndirectLighting &&
+	bAffectDistanceFieldLighting == Other.bAffectDistanceFieldLighting &&
+	bReceivesDecals == Other.bReceivesDecals &&
+	bOverrideLightMapRes == Other.bOverrideLightMapRes &&
+	bUseAsOccluder == Other.bUseAsOccluder &&
+	bRenderCustomDepth == Other.bRenderCustomDepth &&
+	bIncludeInHLOD == Other.bIncludeInHLOD &&
+	bEnableDiscardOnLoad == Other.bEnableDiscardOnLoad &&
+	BodyInstance.GetCollisionEnabled() == Other.BodyInstance.GetCollisionEnabled() && 
+	BodyInstance.GetCollisionResponse() == Other.BodyInstance.GetCollisionResponse() &&
+	BodyInstance.DoesUseCollisionProfile() == Other.BodyInstance.DoesUseCollisionProfile() &&
+	(!BodyInstance.DoesUseCollisionProfile() || (BodyInstance.GetCollisionProfileName() == BodyInstance.GetCollisionProfileName()));
+}
+
+uint32 FISMComponentDescriptor::ComputeHash() const
+{
+	FArchiveObjectCrc32 CrcArchive;
+	UISMComponentDescriptorHasher* Hasher = NewObject<UISMComponentDescriptorHasher>();
+	Hash = 0; // we don't want the hash to impact the calculation
+	Hasher->Descriptor = *this;
+	uint32 Crc = CrcArchive.Crc32(Hasher);
+	Hash = Crc;
+
+	return Crc;
+}
+
+UInstancedStaticMeshComponent* FISMComponentDescriptor::CreateComponent(UObject* Outer, FName Name, EObjectFlags ObjectFlags) const
+{
+	UInstancedStaticMeshComponent* ISMComponent = NewObject<UInstancedStaticMeshComponent>(Outer, ComponentClass, Name, ObjectFlags);
+	
+	ISMComponent->SetStaticMesh(StaticMesh);
+	ISMComponent->OverrideMaterials = OverrideMaterials;
+	ISMComponent->Mobility = Mobility;
+	ISMComponent->RuntimeVirtualTextures = RuntimeVirtualTextures;
+	ISMComponent->VirtualTextureRenderPassType = VirtualTextureRenderPassType;
+	ISMComponent->LightmapType = LightmapType;
+	ISMComponent->LightingChannels = LightingChannels;
+	ISMComponent->bHasCustomNavigableGeometry = bHasCustomNavigableGeometry;
+	ISMComponent->CustomDepthStencilWriteMask = CustomDepthStencilWriteMask;
+	ISMComponent->BodyInstance.CopyBodyInstancePropertiesFrom(&BodyInstance);
+	ISMComponent->InstanceStartCullDistance = InstanceStartCullDistance;
+	ISMComponent->InstanceEndCullDistance = InstanceEndCullDistance;
+	ISMComponent->VirtualTextureCullMips = VirtualTextureCullMips;
+	ISMComponent->TranslucencySortPriority = TranslucencySortPriority;
+	ISMComponent->OverriddenLightMapRes = OverriddenLightMapRes;
+	ISMComponent->CustomDepthStencilValue = CustomDepthStencilValue;
+	ISMComponent->CastShadow = bCastShadow;
+	ISMComponent->bCastStaticShadow = bCastStaticShadow;
+	ISMComponent->bCastDynamicShadow = bCastDynamicShadow;
+	ISMComponent->bCastContactShadow = bCastContactShadow;
+	ISMComponent->bCastShadowAsTwoSided = bCastShadowAsTwoSided;
+	ISMComponent->bAffectDynamicIndirectLighting = bAffectDynamicIndirectLighting;
+	ISMComponent->bAffectDistanceFieldLighting = bAffectDistanceFieldLighting;
+	ISMComponent->bReceivesDecals = bReceivesDecals;
+	ISMComponent->bOverrideLightMapRes = bOverrideLightMapRes;
+	ISMComponent->bUseAsOccluder = bUseAsOccluder;
+	ISMComponent->bRenderCustomDepth = bRenderCustomDepth;
+	ISMComponent->bEnableAutoLODGeneration = bIncludeInHLOD;;
+
+	// HISM Specific
+	if (UHierarchicalInstancedStaticMeshComponent* HISMComponent = Cast<UHierarchicalInstancedStaticMeshComponent>(ISMComponent))
+	{
+		HISMComponent->bEnableDensityScaling = bEnableDensityScaling;
+	}
+
+	return ISMComponent;
+}
+
+#endif
+
+UISMComponentDescriptorHasher::UISMComponentDescriptorHasher(const FObjectInitializer& ObjectInitializer)
+	: Super(ObjectInitializer)
+{
+}
\ No newline at end of file
diff --git a/Engine/Source/Runtime/Engine/Private/ISMPartition/ISMPartitionActor.cpp b/Engine/Source/Runtime/Engine/Private/ISMPartition/ISMPartitionActor.cpp
new file mode 100644
--- /dev/null
+++ b/Engine/Source/Runtime/Engine/Private/ISMPartition/ISMPartitionActor.cpp
@@ -0,0 +1,626 @@
+// Copyright Epic Games, Inc. All Rights Reserved.
+#include "ISMPartition/ISMPartitionActor.h"
+#include "ISMPartition/ISMComponentDescriptor.h"
+#include "Components/HierarchicalInstancedStaticMeshComponent.h"
+#include "Engine/World.h"
+#include "Engine/StaticMesh.h"
+
+#if WITH_EDITOR
+#include "Editor.h"
+#endif
+
+DEFINE_LOG_CATEGORY(LogISMPartition);
+
+AISMPartitionActor::AISMPartitionActor(const FObjectInitializer& ObjectInitializer)
+	: Super(ObjectInitializer)
+#if WITH_EDITORONLY_DATA
+	, UpdateDepth(0)
+	, bWasModifyCalled(false)
+#endif
+{
+
+}
+
+#if WITH_EDITOR
+static FAutoConsoleCommand DumpISMPartitionActors(
+	TEXT("ism.Editor.DumpISMPartitionActors"),
+	TEXT("Output stats about ISMPartitionActor(s)"),
+	FConsoleCommandWithArgsDelegate::CreateLambda([](const TArray<FString>& Args)
+	{
+		if (UWorld* World = GEditor->GetEditorWorldContext().World())
+		{
+			if (!World->IsPlayInEditor())
+			{
+				for (TActorIterator<AISMPartitionActor> It(World); It; ++It)
+				{
+					AISMPartitionActor* ISMPartitionActor = *It;
+					ISMPartitionActor->OutputStats();
+				}
+
+				// Also output total ism count
+				int32 TotalCount = 0;
+				int32 ComponentCount = 0;
+				for (TObjectIterator<UInstancedStaticMeshComponent> It(RF_ClassDefaultObject | RF_ArchetypeObject, true, EInternalObjectFlags::PendingKill); It; ++It)
+				{
+					UInstancedStaticMeshComponent* ISMComponent = *It;
+					if (ISMComponent)
+					{
+						TotalCount += ISMComponent->GetInstanceCount();
+						ComponentCount++;
+					}
+				}
+				UE_LOG(LogISMPartition, Log, TEXT("ISM Component Count: %d Instance Count: %d"), ComponentCount, TotalCount);
+			}
+		}
+	})
+);
+
+
+void AISMPartitionActor::OutputStats() const
+{
+	UE_LOG(LogISMPartition, Log, TEXT("ISM Partition: %s (%s)"), *GetActorLabel(), *GetPathName());
+	int32 TotalCount = 0;
+	for (int32 DescriptorIndex = 0; DescriptorIndex < Descriptors.Num(); ++DescriptorIndex)
+	{
+		int32 Count = DescriptorComponents[DescriptorIndex].Instances.Num();
+		if (Count > 0)
+		{
+			UE_LOG(LogISMPartition, Log, TEXT("StaticMesh: %s Instance Count: %d"), *Descriptors[DescriptorIndex].StaticMesh->GetPathName(), Count);
+		}
+		TotalCount += Count;
+	}
+	UE_LOG(LogISMPartition, Log, TEXT("Total Instance Count: %d"), TotalCount);
+}
+
+void AISMPartitionActor::PreEditUndo()
+{
+	Super::PreEditUndo();
+
+	for (FISMComponentData& ComponentData : DescriptorComponents)
+	{
+		ComponentData.UnregisterDelegates();
+	}
+}
+
+void AISMPartitionActor::PostEditUndo()
+{
+	Super::PostEditUndo();
+
+	for (FISMComponentData& ComponentData : DescriptorComponents)
+	{
+		ComponentData.RegisterDelegates();
+	}
+}
+
+FISMClientHandle AISMPartitionActor::RegisterClient(const FGuid& ClientGuid)
+{
+	Modify();
+	int32 Index = Clients.Find(ClientGuid);
+	if (Index == INDEX_NONE)
+	{
+		Index = Clients.Find(FGuid());
+		if (Index == INDEX_NONE)
+		{
+			Index = Clients.Num();
+			Clients.Add(ClientGuid);
+		}
+		else // Fill hole
+		{
+			Clients[Index] = ClientGuid;
+		}
+	}
+	return FISMClientHandle(Index, ClientGuid);
+}
+
+void AISMPartitionActor::UnregisterClient(FISMClientHandle& Handle)
+{
+	check(Handle.Guid == Clients[Handle.Index]);
+	RemoveISMInstances(Handle);
+			
+	Clients[Handle.Index] = FGuid();
+	Handle.Index = -1;
+	Handle.Guid = FGuid();
+}
+
+int32 AISMPartitionActor::RegisterISMComponentDescriptor(const FISMComponentDescriptor& Descriptor)
+{
+	check(Descriptor.Hash != 0);
+	ModifyActor();
+	int32 Index = Descriptors.IndexOfByPredicate([&](const FISMComponentDescriptor& ExistingDescriptor) { return ExistingDescriptor == Descriptor; });
+	if (Index == INDEX_NONE)
+	{
+		Index = Descriptors.IndexOfByPredicate([&](const FISMComponentDescriptor& ExistingDescriptor) { return ExistingDescriptor.Hash == 0; });
+		if (Index == INDEX_NONE)
+		{
+			check(Descriptors.Num() == DescriptorComponents.Num());
+			Index = Descriptors.Num();
+			Descriptors.Add(Descriptor);
+			DescriptorComponents.Add(FISMComponentData());
+		}
+		else
+		{
+			Descriptors[Index] = Descriptor;
+			check(DescriptorComponents[Index].Instances.Num() == 0);
+		}
+	}
+	return Index;
+}
+
+void AISMPartitionActor::ReserveISMInstances(const FISMClientHandle& Handle, int32 AddedInstanceCount, const TSortedMap<int32, TArray<FTransform>>& InstanceDefinition)
+{
+	check(Handle.Guid == Clients[Handle.Index]);
+	BeginUpdate();
+	ModifyActor();
+	for (const auto& Pair : InstanceDefinition)
+	{
+		const int32 DescriptorIndex = Pair.Key;
+
+		FISMComponentData& ComponentData = DescriptorComponents[DescriptorIndex];
+		if (!ComponentData.Component)
+		{
+			CreateComponent(Descriptors[DescriptorIndex], ComponentData);
+		}
+		ModifyComponent(ComponentData);
+		ComponentData.Component->PreAllocateInstancesMemory(AddedInstanceCount);
+
+		if (ComponentData.ClientInstances.Num() <= Handle.Index)
+		{
+			ComponentData.ClientInstances.AddDefaulted(Handle.Index - ComponentData.ClientInstances.Num() + 1);
+		}
+
+		FISMClientData& ISMClientData = ComponentData.ClientInstances[Handle.Index];
+		ISMClientData.Instances.Reserve(ISMClientData.Instances.Num() + AddedInstanceCount);
+	}
+	EndUpdate();
+}
+
+void AISMPartitionActor::AddISMInstance(const FISMClientHandle& Handle, const FTransform& InstanceTransform, const TSortedMap<int32, TArray<FTransform>>& InstanceDefinition)
+{
+	check(Handle.Guid == Clients[Handle.Index]);
+	// Create Descriptor entry
+	BeginUpdate();
+	ModifyActor();
+
+	for (const auto& Pair : InstanceDefinition)
+	{
+		const int32 DescriptorIndex = Pair.Key;
+		
+		FISMComponentData& ComponentData = DescriptorComponents[DescriptorIndex];
+		if (!ComponentData.Component)
+		{
+			CreateComponent(Descriptors[DescriptorIndex], ComponentData);
+		}
+
+		if (ComponentData.ClientInstances.Num() <= Handle.Index)
+		{
+			ComponentData.ClientInstances.AddDefaulted(Handle.Index - ComponentData.ClientInstances.Num() + 1);
+		}
+
+		FISMClientData& ClientData = ComponentData.ClientInstances[Handle.Index];
+		int32 ClientIndex = ClientData.Instances.Num();
+		FISMClientInstance& ClientInstance = ClientData.Instances.Emplace_GetRef();
+		ClientInstance.ComponentIndices.Reserve(Pair.Value.Num());
+		ComponentData.Instances.Reserve(ComponentData.Instances.Num() + Pair.Value.Num());
+
+		// Add Instances and creating mappings for faster removal
+		for (const FTransform& LocalTransform : Pair.Value)
+		{
+			int32 ComponentInstanceIndex = ComponentData.Instances.Num();
+			ComponentData.Instances.Emplace(Handle.Index, ClientIndex, ClientInstance.ComponentIndices.Num());			
+			ClientInstance.ComponentIndices.Add(ComponentInstanceIndex);
+			
+			FTransform NewInstanceTransform = LocalTransform * InstanceTransform;
+			AddInstanceToComponent(ComponentData, NewInstanceTransform);
+		}
+	}
+
+	EndUpdate();
+}
+
+void AISMPartitionActor::RemoveISMInstances(const FISMClientHandle& Handle)
+{
+	check(Handle.Guid == Clients[Handle.Index]);
+	BeginUpdate();
+	
+	for (int32 DescriptorIndex = Descriptors.Num()-1; DescriptorIndex >= 0; --DescriptorIndex)
+	{
+		FISMComponentData& ComponentData = DescriptorComponents[DescriptorIndex];
+		if (ComponentData.ClientInstances.IsValidIndex(Handle.Index))
+		{
+			FISMClientData& ClientData = ComponentData.ClientInstances[Handle.Index];
+			for (int32 InstanceIndex = ClientData.Instances.Num() - 1; InstanceIndex >= 0; --InstanceIndex)
+			{
+				RemoveISMInstancesInternal(ComponentData, ClientData, InstanceIndex);
+			}
+			ClientData.Instances.Empty(); // Free the memory
+		}
+
+		bool bIsEmpty = DestroyComponentIfEmpty(Descriptors[DescriptorIndex], ComponentData);
+		// If Component is empty and its the last descriptor in the list we can remove it without breaking indices
+		if (bIsEmpty && DescriptorIndex == Descriptors.Num() - 1)
+		{
+			Descriptors.RemoveAt(DescriptorIndex);
+			DescriptorComponents.RemoveAt(DescriptorIndex);
+		}
+	}
+		
+	EndUpdate();
+}
+
+bool AISMPartitionActor::DestroyComponentIfEmpty(FISMComponentDescriptor& Descriptor, FISMComponentData& ComponentData)
+{
+	if (ComponentData.Component && !ComponentData.Component->GetInstanceCount())
+	{
+		ModifyActor();
+		Descriptor.Hash = 0; // FreeList
+		ComponentData.Instances.Empty();
+		ComponentData.UnregisterDelegates();
+		ComponentData.Component->Modify();
+		ComponentData.Component->DestroyComponent();
+		ComponentData.Component = nullptr;
+		return true;
+	}
+
+	return false;
+}
+
+void AISMPartitionActor::InvalidateComponentLightingCache(FISMComponentData& ComponentData)
+{
+	if (UpdateDepth)
+	{
+		ComponentData.bInvalidateLightingCache = true;
+	}
+	else
+	{
+		ComponentData.Component->InvalidateLightingCache();
+		ComponentData.bInvalidateLightingCache = false;
+	}
+}
+
+void AISMPartitionActor::RemoveInstanceFromComponent(FISMComponentData& ComponentData, int32 ComponentInstanceIndex)
+{
+	if (ComponentData.Component)
+	{
+		ModifyComponent(ComponentData);
+		ComponentData.Component->RemoveInstance(ComponentInstanceIndex);
+		InvalidateComponentLightingCache(ComponentData);
+	}
+}
+
+void AISMPartitionActor::AddInstanceToComponent(FISMComponentData& ComponentData, const FTransform& WorldTransform)
+{
+	check(ComponentData.Component);
+	ModifyComponent(ComponentData);
+	ComponentData.Component->AddInstanceWorldSpace(WorldTransform);
+	InvalidateComponentLightingCache(ComponentData);
+}
+
+void AISMPartitionActor::UpdateInstanceTransform(FISMComponentData& ComponentData, int32 ComponentInstanceIndex, const FTransform& WorldTransform, bool bTeleport)
+{
+	check(ComponentData.Component);
+	ModifyComponent(ComponentData);
+	ComponentData.Component->UpdateInstanceTransform(ComponentInstanceIndex, WorldTransform, true, true, bTeleport);
+	InvalidateComponentLightingCache(ComponentData);
+}
+
+void AISMPartitionActor::ModifyComponent(FISMComponentData& ComponentData)
+{
+	check(ComponentData.Component);
+	if (!ComponentData.bWasModifyCalled || UpdateDepth == 0)
+	{
+		// Only store the bModfied inside a BeginUpdate/EndUpdate
+		ComponentData.bWasModifyCalled = UpdateDepth > 0;
+		ComponentData.Component->Modify();
+	}
+}
+
+void AISMPartitionActor::CreateComponent(const FISMComponentDescriptor& ComponentDescriptor, FISMComponentData& ComponentData)
+{
+	ModifyActor();
+
+	check(!ComponentData.Component);
+	ComponentData.Component = ComponentDescriptor.CreateComponent(this, NAME_None, RF_Transactional);
+	ComponentData.Component->bSelectable = true;
+	ComponentData.Component->bHasPerInstanceHitProxies = true;
+	ComponentData.RegisterDelegates();
+
+	ComponentData.Component->SetupAttachment(RootComponent);
+
+	if (GetRootComponent()->IsRegistered())
+	{
+		ComponentData.Component->RegisterComponent();
+	}
+
+	// Use only instance translation as a component transform
+	ComponentData.Component->SetWorldTransform(RootComponent->GetComponentTransform());
+
+	// Add the new component to the transaction buffer so it will get destroyed on undo
+	ComponentData.Component->Modify();
+
+	check(UpdateDepth > 0);
+	if (UHierarchicalInstancedStaticMeshComponent* HISMComponent = Cast<UHierarchicalInstancedStaticMeshComponent>(ComponentData.Component))
+	{
+		ComponentData.bAutoRebuildTreeOnInstanceChanges = HISMComponent->bAutoRebuildTreeOnInstanceChanges;
+		HISMComponent->bAutoRebuildTreeOnInstanceChanges = false;
+	}
+}
+
+void AISMPartitionActor::ModifyActor()
+{
+	if (!bWasModifyCalled || UpdateDepth == 0)
+	{
+		bWasModifyCalled = UpdateDepth > 0;
+		Modify();
+	}
+}
+
+void AISMPartitionActor::RemoveISMInstancesInternal(FISMComponentData& ComponentData, FISMClientData& ClientData, int32 InstanceIndex)
+{
+	if (!ClientData.Instances.Num())
+	{
+		return;
+	}
+
+	BeginUpdate();
+	ModifyActor();
+		
+	check(ComponentData.Component);
+	const bool bRemoveSwap = ComponentData.Component->SupportsRemoveSwap();
+	FISMClientInstance& ClientInstance = ClientData.Instances[InstanceIndex];
+	for (int32 ComponentInstanceIndex : ClientInstance.ComponentIndices)
+	{
+		bool OutIsRemoveAtSwap = false;
+		RemoveInstanceFromComponent(ComponentData, ComponentInstanceIndex);
+
+		if (bRemoveSwap)
+		{
+			ComponentData.Instances.RemoveAtSwap(ComponentInstanceIndex, 1, false);
+
+			// Make sure we have the proper index (if we removed the last element, or the only element)
+			ComponentInstanceIndex = FMath::Min(ComponentInstanceIndex, ComponentData.Instances.Num() - 1);
+			if (ComponentInstanceIndex >= 0)
+			{
+				// Fix up because of swapping
+				const FISMComponentInstance& ComponentInstance = ComponentData.Instances[ComponentInstanceIndex];
+				FISMClientData& SwapClientData = ComponentData.ClientInstances[ComponentInstance.ClientIndex];
+				SwapClientData.Instances[ComponentInstance.InstanceIndex].ComponentIndices[ComponentInstance.InstanceSubIndex] = ComponentInstanceIndex;
+			}
+		}
+		else
+		{
+			ComponentData.Instances.RemoveAt(ComponentInstanceIndex);
+
+			// No swap means we need to update all instances after the removed instance
+			for (int32 IndexToUpdate = ComponentInstanceIndex; IndexToUpdate < ComponentData.Instances.Num(); ++IndexToUpdate)
+			{
+				const FISMComponentInstance& ComponentInstance = ComponentData.Instances[IndexToUpdate];
+				FISMClientData& SwapClientData = ComponentData.ClientInstances[ComponentInstance.ClientIndex];
+				SwapClientData.Instances[ComponentInstance.InstanceIndex].ComponentIndices[ComponentInstance.InstanceSubIndex] = IndexToUpdate;
+			}
+		}
+	}
+
+	ClientData.Instances.RemoveAtSwap(InstanceIndex, 1, false);
+	if (ClientData.Instances.IsValidIndex(InstanceIndex))
+	{
+		FISMClientInstance& SwapClientInstance = ClientData.Instances[InstanceIndex];
+		for (int32 ComponentInstanceIndex : SwapClientInstance.ComponentIndices)
+		{
+			FISMComponentInstance& ComponentInstance = ComponentData.Instances[ComponentInstanceIndex];
+			ComponentInstance.InstanceIndex = InstanceIndex;
+		}
+	}
+
+	EndUpdate();
+}
+
+void AISMPartitionActor::RemoveISMInstance(const FISMClientHandle& Handle, int32 InstanceIndex)
+{
+	check(Handle.Guid == Clients[Handle.Index]);
+	for (FISMComponentData& ComponentData : DescriptorComponents)
+	{
+		if (ComponentData.ClientInstances.IsValidIndex(Handle.Index))
+		{
+			FISMClientData& ClientData = ComponentData.ClientInstances[Handle.Index];
+			
+			RemoveISMInstancesInternal(ComponentData, ClientData, InstanceIndex);
+		}
+	}
+}
+
+void AISMPartitionActor::SelectISMInstances(const FISMClientHandle& Handle, bool bSelect, const TSet<int32>& Indices)
+{
+	check(Handle.Guid == Clients[Handle.Index]);
+	if (Indices.Num())
+	{
+		for (FISMComponentData& ComponentData : DescriptorComponents)
+		{
+			if (ComponentData.ClientInstances.IsValidIndex(Handle.Index))
+			{
+				FISMClientData& ClientData = ComponentData.ClientInstances[Handle.Index];
+				if (ClientData.Instances.Num())
+				{
+					ModifyComponent(ComponentData);
+
+					for (int32 SelectIndex : Indices)
+					{
+						FISMClientInstance& ClientInstance = ClientData.Instances[SelectIndex];
+						for (int32 ComponentInstanceIndex : ClientInstance.ComponentIndices)
+						{
+							ComponentData.Component->SelectInstance(bSelect, ComponentInstanceIndex);
+						}
+					}
+				}
+			}
+		}
+	}
+}
+
+void AISMPartitionActor::SetISMInstanceTransform(const FISMClientHandle& Handle, int32 InstanceIndex, const FTransform& NewTransform, bool bTeleport, const TSortedMap<int32, TArray<FTransform>>& InstanceDefinition)
+{
+	check(Handle.Guid == Clients[Handle.Index]);
+	BeginUpdate();
+
+	// We don't store the previous transform so for now we need the Definition to compute the new world positions
+	for (const auto& Pair : InstanceDefinition)
+	{
+		const int32 DescriptorIndex = Pair.Key;
+
+		FISMComponentData& ComponentData = DescriptorComponents[DescriptorIndex];
+		FISMClientData& ClientData = ComponentData.ClientInstances[Handle.Index];
+		FISMClientInstance& ClientInstance = ClientData.Instances[InstanceIndex];
+		check(ClientInstance.ComponentIndices.Num() == Pair.Value.Num());
+
+		// Add Instances and creating mappings for faster removal
+		for (int32 Index = 0; Index < ClientInstance.ComponentIndices.Num(); ++Index)
+		{
+			const FTransform& LocalTransform = Pair.Value[Index];
+			
+			int32 ComponentInstanceIndex = ClientInstance.ComponentIndices[Index];
+
+			UpdateInstanceTransform(ComponentData, ComponentInstanceIndex, LocalTransform * NewTransform, bTeleport);
+		}
+	}
+
+	EndUpdate();
+}
+
+bool AISMPartitionActor::IsISMComponent(const UPrimitiveComponent* Component) const
+{
+	for (const FISMComponentData& ComponentData : DescriptorComponents)
+	{
+		if (ComponentData.Component == Component)
+		{
+			return true;
+		}
+	}
+
+	return false;
+}
+
+int32 AISMPartitionActor::GetISMInstanceIndex(const FISMClientHandle& Handle, UInstancedStaticMeshComponent* ISMComponent, int32 ComponentIndex) const
+{
+	check(Handle.Guid == Clients[Handle.Index]);
+	for (const FISMComponentData& ComponentData : DescriptorComponents)
+	{
+		if (ComponentData.Component == ISMComponent)
+		{
+			const FISMComponentInstance& ComponentInstance = ComponentData.Instances[ComponentIndex];
+			if (ComponentInstance.ClientIndex == Handle.Index)
+			{
+				return ComponentInstance.InstanceIndex;
+			}
+			
+			break;
+		}
+	}
+
+	return INDEX_NONE;
+}
+
+FBox AISMPartitionActor::GetISMInstanceBounds(const FISMClientHandle& Handle, const TSet<int32>& Indices) const
+{
+	check(Handle.Guid == Clients[Handle.Index]);
+	FBox BoundingBox(EForceInit::ForceInit);
+	FTransform InstanceWorldTransform;
+	for (int32 i : Indices)
+	{
+		for (const FISMComponentData& ComponentData : DescriptorComponents)
+		{
+			if (ComponentData.Component && ComponentData.Component->GetStaticMesh())
+			{
+				const FBox StaticMeshBoundingBox = ComponentData.Component->GetStaticMesh()->GetBoundingBox();
+				const FISMClientData& ClientData = ComponentData.ClientInstances[Handle.Index];
+				for (const FISMClientInstance& ClientInstance : ClientData.Instances)
+				{
+					for (int32 ComponentInstanceIndex : ClientInstance.ComponentIndices)
+					{
+						ComponentData.Component->GetInstanceTransform(ComponentInstanceIndex, InstanceWorldTransform, true);
+						BoundingBox += StaticMeshBoundingBox.TransformBy(InstanceWorldTransform);
+					}
+				}
+			}
+		}
+
+		
+	}
+	return BoundingBox;
+}
+
+void AISMPartitionActor::BeginUpdate()
+{
+	if (UpdateDepth == 0)
+	{
+		for (FISMComponentData& ComponentData : DescriptorComponents)
+		{
+			if (UHierarchicalInstancedStaticMeshComponent* HISMComponent = Cast<UHierarchicalInstancedStaticMeshComponent>(ComponentData.Component))
+			{
+				ComponentData.bAutoRebuildTreeOnInstanceChanges = HISMComponent->bAutoRebuildTreeOnInstanceChanges;
+				HISMComponent->bAutoRebuildTreeOnInstanceChanges = false;
+			}
+
+			check(!ComponentData.bWasModifyCalled);
+			check(!bWasModifyCalled);
+		}
+	}
+	++UpdateDepth;
+}
+
+void AISMPartitionActor::EndUpdate()
+{
+	--UpdateDepth;
+	check(UpdateDepth >= 0);
+	if (UpdateDepth == 0)
+	{
+		bWasModifyCalled = false;
+		for (FISMComponentData& ComponentData : DescriptorComponents)
+		{
+			
+			if (ComponentData.Component && ComponentData.bInvalidateLightingCache)
+			{
+				ComponentData.Component->InvalidateLightingCache();
+				ComponentData.bInvalidateLightingCache = false;
+			}
+			if (UHierarchicalInstancedStaticMeshComponent* HISMComponent = Cast<UHierarchicalInstancedStaticMeshComponent>(ComponentData.Component))
+			{
+				HISMComponent->bAutoRebuildTreeOnInstanceChanges = ComponentData.bAutoRebuildTreeOnInstanceChanges;
+
+				if (ComponentData.bWasModifyCalled)
+				{
+					HISMComponent->BuildTreeIfOutdated(true, false);
+				}
+			}
+
+			ComponentData.bWasModifyCalled = false;
+		}
+	}
+}
+
+void AISMPartitionActor::UpdateHISMTrees(bool bAsync, bool bForce)
+{
+	for (FISMComponentData& ComponentData : DescriptorComponents)
+	{
+		if (UHierarchicalInstancedStaticMeshComponent* HISMComponent = Cast<UHierarchicalInstancedStaticMeshComponent>(ComponentData.Component))
+		{
+			HISMComponent->BuildTreeIfOutdated(bAsync, bForce);
+		}
+	}
+}
+
+void AISMPartitionActor::GetClientComponents(const FISMClientHandle& Handle, TArray<UInstancedStaticMeshComponent*>& OutComponents)
+{
+	check(Handle.Guid == Clients[Handle.Index]);
+	for (FISMComponentData& ComponentData : DescriptorComponents)
+	{
+		if (ComponentData.Component)
+		{
+			FISMClientData& ClientData = ComponentData.ClientInstances[Handle.Index];
+			if (ClientData.Instances.Num())
+			{
+				OutComponents.Add(ComponentData.Component);
+			}
+		}
+	}
+}
+
+#endif
diff --git a/Engine/Source/Runtime/Engine/Public/ISMPartition/ISMComponentData.h b/Engine/Source/Runtime/Engine/Public/ISMPartition/ISMComponentData.h
new file mode 100644
--- /dev/null
+++ b/Engine/Source/Runtime/Engine/Public/ISMPartition/ISMComponentData.h
@@ -0,0 +1,100 @@
+// Copyright Epic Games, Inc. All Rights Reserved.
+#pragma once
+
+#include "CoreMinimal.h"
+#include "Engine/EngineTypes.h"
+#include "ISMComponentData.generated.h"
+
+/** Represents one component instance with a mapping back to the Client instance. The should be a 1 to 1 relationship between those and the actual Component instances. (equal count) */
+USTRUCT()
+struct FISMComponentInstance
+{
+	GENERATED_USTRUCT_BODY()
+
+	FISMComponentInstance(const int32& InClientIndex, const int32& InInstanceIndex, const int32& InInstanceSubIndex)
+		: ClientIndex(InClientIndex)
+		, InstanceIndex(InInstanceIndex)
+		, InstanceSubIndex(InInstanceSubIndex)
+	{
+
+	}
+
+	FISMComponentInstance()
+		: ClientIndex(-1)
+		, InstanceIndex(-1)
+		, InstanceSubIndex(-1)
+	{
+
+	}
+
+	/** Client Index in the AISMPartitionActor */
+	UPROPERTY()
+	int32 ClientIndex;
+
+	/** Instance Index in the FISMClientData struct */
+	UPROPERTY()
+	int32 InstanceIndex;
+
+	/** Instance Index in the FISMClientInstance struct */
+	UPROPERTY()
+	int32 InstanceSubIndex;
+};
+
+/** Represents the component instances that 1 Client instance represents */
+USTRUCT()
+struct FISMClientInstance
+{
+	GENERATED_USTRUCT_BODY()
+
+	/** Instance Index in the ISM Component */
+	UPROPERTY()
+	TArray<int32> ComponentIndices;
+};
+
+/** Represents the list of instances for 1 Client */
+USTRUCT()
+struct FISMClientData
+{
+	GENERATED_USTRUCT_BODY()
+
+	/** Instance list for 1 client */
+	UPROPERTY()
+	TArray<FISMClientInstance> Instances;
+};
+
+/** Bookkeeping struct that contains the data that allows linking Client Instances to actual Component Instances */
+USTRUCT()
+struct FISMComponentData
+{
+	GENERATED_USTRUCT_BODY()
+
+	FISMComponentData();
+
+#if WITH_EDITOR
+	void RegisterDelegates();
+	void UnregisterDelegates();
+	void HandleComponentMeshBoundsChanged(const FBoxSphereBounds& NewBounds);
+#endif
+
+#if WITH_EDITORONLY_DATA
+	/** Instances here must match Component instances */
+	UPROPERTY()
+	TArray<FISMComponentInstance> Instances;
+
+	/** Per Client Instances. Indexed using a Client index that must match the indexing in AISMPartitionActor */
+	UPROPERTY()
+	TArray<FISMClientData> ClientInstances;
+
+	UPROPERTY()
+	UInstancedStaticMeshComponent* Component;
+
+	/** If Lighting cache should be invalidated at the end of a BeginUpdate/EndUpdate */
+	bool bInvalidateLightingCache;
+
+	/** Component's property value before setting it to false in the BeginUpdate so we can restore it in EndUpdate */
+	bool bAutoRebuildTreeOnInstanceChanges;
+
+	/** If Modified as already been called between a BeginUpdate/EndUpdate (avoid multiple Modify calls on the component) */
+	bool bWasModifyCalled;
+#endif
+};
\ No newline at end of file
diff --git a/Engine/Source/Runtime/Engine/Public/ISMPartition/ISMComponentDescriptor.h b/Engine/Source/Runtime/Engine/Public/ISMPartition/ISMComponentDescriptor.h
new file mode 100644
--- /dev/null
+++ b/Engine/Source/Runtime/Engine/Public/ISMPartition/ISMComponentDescriptor.h
@@ -0,0 +1,157 @@
+// Copyright Epic Games, Inc. All Rights Reserved.
+#pragma once
+
+#include "CoreMinimal.h"
+#include "Engine/EngineTypes.h"
+#include "PhysicsEngine/BodyInstance.h"
+#include "Components/PrimitiveComponent.h"
+#include "ISMComponentDescriptor.generated.h"
+
+class UInstancedStaticMeshComponent;
+class UStaticMesh;
+enum class ERuntimeVirtualTextureMainPassType : uint8;
+enum class ERendererStencilMask : uint8;
+
+/** Struct that holds the relevant properties that can help decide if instances of different 
+	StaticMeshComponents/InstancedStaticMeshComponents/HISM can be merged into a single component. */
+USTRUCT()
+struct ENGINE_API FISMComponentDescriptor
+{
+	GENERATED_USTRUCT_BODY()
+
+#if WITH_EDITOR
+	static FISMComponentDescriptor CreateFrom(const TSubclassOf<UStaticMeshComponent>& ComponentClass);
+	void InitFrom(UStaticMeshComponent* Component);
+
+	uint32 ComputeHash() const;
+	UInstancedStaticMeshComponent* CreateComponent(UObject* Outer, FName Name = NAME_None, EObjectFlags ObjectFlags = EObjectFlags::RF_NoFlags) const;
+
+	friend inline uint32 GetTypeHash(const FISMComponentDescriptor& Key)
+	{
+		if (Key.Hash == 0)
+		{
+			Key.ComputeHash();
+		}
+		return Key.Hash;
+	}
+
+	bool operator!=(const FISMComponentDescriptor& Other) const;
+	bool operator==(const FISMComponentDescriptor& Other) const;
+
+	friend inline bool operator<(const FISMComponentDescriptor& Lhs, const FISMComponentDescriptor& Rhs)
+	{
+		return Lhs.Hash < Rhs.Hash;
+	}
+#endif
+
+public:
+
+#if WITH_EDITORONLY_DATA
+	UPROPERTY()
+	mutable uint32 Hash = 0;
+
+	UPROPERTY()
+	TSubclassOf<UInstancedStaticMeshComponent> ComponentClass;
+	
+	UPROPERTY()
+	UStaticMesh* StaticMesh = nullptr;
+	
+	UPROPERTY()
+	TArray<UMaterialInterface*> OverrideMaterials;
+	
+	UPROPERTY()
+	TArray<URuntimeVirtualTexture*> RuntimeVirtualTextures;
+
+	UPROPERTY()
+	TEnumAsByte<EComponentMobility::Type> Mobility;
+		
+	UPROPERTY()
+	ERuntimeVirtualTextureMainPassType VirtualTextureRenderPassType;
+	
+	UPROPERTY()
+	ELightmapType LightmapType;
+
+	UPROPERTY()
+	FLightingChannels LightingChannels;
+
+	UPROPERTY()
+	TEnumAsByte<EHasCustomNavigableGeometry::Type> bHasCustomNavigableGeometry;
+
+	UPROPERTY()
+	ERendererStencilMask CustomDepthStencilWriteMask;
+
+	UPROPERTY()
+	FBodyInstance BodyInstance;
+		
+	UPROPERTY()
+	int32 InstanceStartCullDistance;
+
+	UPROPERTY()
+	int32 InstanceEndCullDistance;
+
+	UPROPERTY()
+	int32 VirtualTextureCullMips;
+
+	UPROPERTY()
+	int32 TranslucencySortPriority;
+
+	UPROPERTY()
+	int32 OverriddenLightMapRes;
+
+	UPROPERTY()
+	int32 CustomDepthStencilValue;
+	
+	UPROPERTY()
+	uint8 bCastShadow : 1;
+
+	UPROPERTY()
+	uint8 bCastDynamicShadow : 1;
+
+	UPROPERTY()
+	uint8 bCastStaticShadow : 1;
+
+	UPROPERTY()
+	uint8 bCastContactShadow : 1;
+
+	UPROPERTY()
+	uint8 bCastShadowAsTwoSided : 1;
+
+	UPROPERTY()
+	uint8 bAffectDynamicIndirectLighting : 1;
+
+	UPROPERTY()
+	uint8 bAffectDistanceFieldLighting : 1;
+
+	UPROPERTY()
+	uint8 bReceivesDecals : 1;
+
+	UPROPERTY()
+	uint8 bOverrideLightMapRes : 1;
+
+	UPROPERTY()
+	uint8 bUseAsOccluder : 1;
+
+	UPROPERTY()
+	uint8 bEnableDensityScaling : 1;
+
+	UPROPERTY()
+	uint8 bEnableDiscardOnLoad : 1;
+
+	UPROPERTY()
+	uint8 bRenderCustomDepth : 1;
+
+	UPROPERTY()
+	uint8 bIncludeInHLOD : 1;
+#endif
+};
+
+/** Utility object to hash the FISMComponentDescriptor through FArchiveObjectCrc32 */
+UCLASS()
+class UISMComponentDescriptorHasher : public UObject
+{
+	GENERATED_UCLASS_BODY()
+
+public:
+	UPROPERTY()
+	FISMComponentDescriptor Descriptor;
+};
diff --git a/Engine/Source/Runtime/Engine/Public/ISMPartition/ISMPartitionActor.h b/Engine/Source/Runtime/Engine/Public/ISMPartition/ISMPartitionActor.h
new file mode 100644
--- /dev/null
+++ b/Engine/Source/Runtime/Engine/Public/ISMPartition/ISMPartitionActor.h
@@ -0,0 +1,119 @@
+// Copyright Epic Games, Inc. All Rights Reserved.
+#pragma once
+
+#include "CoreMinimal.h"
+#include "ActorPartition/PartitionActor.h"
+#include "Misc/Guid.h"
+#include "ISMPartition/ISMComponentDescriptor.h"
+#include "ISMPartition/ISMComponentData.h"
+#include "Templates/Tuple.h"
+#include "Containers/SortedMap.h"
+#include "ISMPartitionActor.generated.h"
+
+struct FISMClientHandle;
+
+DECLARE_LOG_CATEGORY_EXTERN(LogISMPartition, Log, All);
+
+/** Actor base class for instance containers placed on a grid.
+	See UActorPartitionSubsystem. */
+UCLASS(NotPlaceable)
+class ENGINE_API AISMPartitionActor : public APartitionActor
+{
+	GENERATED_UCLASS_BODY()
+
+public:
+#if WITH_EDITOR	
+	//~ Begin AActor Interface
+	virtual void PreEditUndo() override;
+	virtual void PostEditUndo() override;
+	//~ End AActor Interface	
+
+	FISMClientHandle RegisterClient(const FGuid& ClientGuid);
+	void UnregisterClient(FISMClientHandle& ClientHandle);
+
+	int32 RegisterISMComponentDescriptor(const FISMComponentDescriptor& Descriptor);
+	const FISMComponentDescriptor& GetISMComponentDescriptor(int32 DescriptorIndex) const { return Descriptors[DescriptorIndex]; }
+
+	void AddISMInstance(const FISMClientHandle& Handle, const FTransform& InstanceTransform, const TSortedMap<int32, TArray<FTransform>>& InstanceDefinition);
+	void RemoveISMInstance(const FISMClientHandle& Handle, int32 InstanceIndex);
+	void RemoveISMInstances(const FISMClientHandle& Handle);
+	void SelectISMInstances(const FISMClientHandle& Handle, bool bSelect, const TSet<int32>& Indices);
+	void SetISMInstanceTransform(const FISMClientHandle& Handle, int32 InstanceIndex, const FTransform& NewTransform, bool bTeleport, const TSortedMap<int32, TArray<FTransform>>& InstanceDefinition);
+	int32 GetISMInstanceIndex(const FISMClientHandle& Handle, UInstancedStaticMeshComponent* ISMComponent, int32 ComponentIndex) const;
+	FBox GetISMInstanceBounds(const FISMClientHandle& Handle, const TSet<int32>& Indices) const;
+	void ReserveISMInstances(const FISMClientHandle& Handle, int32 AddedInstanceCount, const TSortedMap<int32, TArray<FTransform>>& InstanceDefinition);
+
+	bool IsISMComponent(const UPrimitiveComponent* Component) const;
+	void BeginUpdate();
+	void EndUpdate();
+	void UpdateHISMTrees(bool bAsync, bool bForce);
+	void GetClientComponents(const FISMClientHandle& Handle, TArray<UInstancedStaticMeshComponent*>& OutComponents);
+	void OutputStats() const;
+protected:
+
+private:
+	void RemoveISMInstancesInternal(FISMComponentData& ComponentData, FISMClientData& OwnerData, int32 InstanceIndex);
+
+	void InvalidateComponentLightingCache(FISMComponentData& ComponentData);
+	void AddInstanceToComponent(FISMComponentData& ComponentData, const FTransform& WorldTransform);
+	void UpdateInstanceTransform(FISMComponentData& ComponentData, int32 ComponentInstanceIndex, const FTransform& WorldTransform, bool bTeleport);
+	void RemoveInstanceFromComponent(FISMComponentData& ComponentData, int32 ComponentInstanceIndex);
+	bool DestroyComponentIfEmpty(FISMComponentDescriptor& Descriptor, FISMComponentData& ComponentData);
+	void ModifyComponent(FISMComponentData& ComponentData);
+	void CreateComponent(const FISMComponentDescriptor& Descriptor, FISMComponentData& ComponentData);
+	void ModifyActor();
+#endif
+
+#if WITH_EDITORONLY_DATA
+	UPROPERTY()
+	TArray<FGuid> Clients;
+		
+	UPROPERTY()
+	TArray<FISMComponentDescriptor> Descriptors;
+
+	UPROPERTY()
+	TArray<FISMComponentData> DescriptorComponents;
+
+	/** If greater than 0 means we are between a BeginUpdate/EndUpdate call and there are some things we can delay/optimize */
+	int32 UpdateDepth;
+
+	/** If Modified as already been called between a BeginUpdate/EndUpdate (avoid multiple Modify calls on the component) */
+	bool bWasModifyCalled;
+#endif
+};
+
+USTRUCT()
+struct FISMClientHandle
+{
+	GENERATED_USTRUCT_BODY();
+
+	FISMClientHandle()
+		: Index(-1)
+	{
+
+	}
+
+	bool IsValid() const { return Index >= 0 && Guid.IsValid(); }
+
+	void Serialize(FArchive& Ar)
+	{
+		Ar << Index;
+		Ar << Guid;
+	}
+
+private:
+	friend class AISMPartitionActor;
+
+	FISMClientHandle(int32 ClientIndex, FGuid ClientGuid)
+		: Index(ClientIndex)
+		, Guid(ClientGuid)
+	{
+
+	}
+
+	UPROPERTY()
+	int32 Index;
+
+	UPROPERTY()
+	FGuid Guid;
+};
\ No newline at end of file
diff --git a/Engine/Source/Runtime/Foliage/Private/FoliageActor.cpp b/Engine/Source/Runtime/Foliage/Private/FoliageActor.cpp
--- a/Engine/Source/Runtime/Foliage/Private/FoliageActor.cpp
+++ b/Engine/Source/Runtime/Foliage/Private/FoliageActor.cpp
@@ -212,13 +212,13 @@ int32 FFoliageActor::FindIndex(const UPrimitiveComponent* Component) const
 	return FindIndex(Component->GetOwner());
 }
 
-void FFoliageActor::Refresh(const TArray<FFoliageInstance>& Instances, bool Async, bool Force)
+void FFoliageActor::Refresh(bool Async, bool Force)
 {
-	for (int32 i = 0; i < Instances.Num(); ++i)
+	for (int32 i = 0; i < Info->Instances.Num(); ++i)
 	{
 		if (ActorInstances[i] == nullptr || ActorInstances[i]->IsPendingKill())
 		{
-			ActorInstances[i] = Spawn(Instances[i]);
+			ActorInstances[i] = Spawn(Info->Instances[i]);
 		}
 	}
 }
@@ -246,10 +246,12 @@ void FFoliageActor::UpdateActorTransforms(const TArray<FFoliageInstance>& Instan
 	}
 }
 
-void FFoliageActor::PostEditUndo(FFoliageInfo* InInfo, UFoliageType* FoliageType, const TArray<FFoliageInstance>& Instances, const TSet<int32>& SelectedIndices)
+void FFoliageActor::PostEditUndo(FFoliageInfo* InInfo, UFoliageType* FoliageType)
 {
-	FFoliageImpl::PostEditUndo(InInfo, FoliageType, Instances, SelectedIndices);
-	UpdateActorTransforms(Instances);
+	// This will set Info to InInfo (as its probably nullptr after an undo operation)
+	FFoliageImpl::PostEditUndo(InInfo, FoliageType);
+	check(Info == InInfo);
+	UpdateActorTransforms(Info->Instances);
 }
 
 void FFoliageActor::PreMoveInstances(const TArray<int32>& InInstancesMoved)
@@ -283,11 +285,23 @@ void FFoliageActor::PostMoveInstances(const TArray<int32>& InInstancesMoved, boo
 	}
 }
 
-void FFoliageActor::NotifyFoliageTypeChanged(UFoliageType* FoliageType, const TArray<FFoliageInstance>& Instances, const TSet<int32>& SelectedIndices, bool bSourceChanged)
+void FFoliageActor::NotifyFoliageTypeChanged(UFoliageType* FoliageType, bool bSourceChanged)
 {
-	AInstancedFoliageActor* IFA = GetIFA();
+	if (!IsInitialized())
+	{
+		return;
+	}
+
 	if (UFoliageType_Actor* FoliageTypeActor = Cast<UFoliageType_Actor>(FoliageType))
 	{
+		// Implementation should change
+		if (FoliageTypeActor->bStaticMeshOnly)
+		{
+			Uninitialize();
+			return;
+		}
+
+		AInstancedFoliageActor* IFA = GetIFA();
 		if (bShouldAttachToBaseComponent != FoliageTypeActor->bShouldAttachToBaseComponent)
 		{
 			bShouldAttachToBaseComponent = FoliageTypeActor->bShouldAttachToBaseComponent;
@@ -300,16 +314,16 @@ void FFoliageActor::NotifyFoliageTypeChanged(UFoliageType* FoliageType, const TA
 
 	if (bSourceChanged)
 	{
-		Reapply(FoliageType, Instances);
-		ApplySelection(true, SelectedIndices);
+		Reapply(FoliageType);
+		ApplySelection(true, Info->SelectedIndices);
 	}
 }
 
-void FFoliageActor::Reapply(const UFoliageType* FoliageType, const TArray<FFoliageInstance>& Instances, bool bPostLoad)
+void FFoliageActor::Reapply(const UFoliageType* FoliageType)
 {
 	AInstancedFoliageActor* IFA = GetIFA();
 	IFA->Modify();
-	DestroyActors(bPostLoad);
+	DestroyActors(false);
 	
 	if (IsInitialized())
 	{
@@ -317,9 +331,9 @@ void FFoliageActor::Reapply(const UFoliageType* FoliageType, const TArray<FFolia
 	}
 	Initialize(FoliageType);
 
-	for (int32 i = 0; i < Instances.Num(); ++i)
+	for (int32 i = 0; i < Info->Instances.Num(); ++i)
 	{
-		ActorInstances.Add(Spawn(Instances[i]));
+		ActorInstances.Add(Spawn(Info->Instances[i]));
 	}
 }
 
diff --git a/Engine/Source/Runtime/Foliage/Private/FoliageActor.h b/Engine/Source/Runtime/Foliage/Private/FoliageActor.h
--- a/Engine/Source/Runtime/Foliage/Private/FoliageActor.h
+++ b/Engine/Source/Runtime/Foliage/Private/FoliageActor.h
@@ -52,13 +52,13 @@ struct FFoliageActor : public FFoliageImpl
 
 	virtual void BeginUpdate() override;
 	virtual void EndUpdate() override;
-	virtual void Refresh(const TArray<FFoliageInstance>& Instances, bool Async, bool Force) override;
+	virtual void Refresh(bool Async, bool Force) override;
 	virtual void OnHiddenEditorViewMaskChanged(uint64 InHiddenEditorViews) override;
-	virtual void PostEditUndo(FFoliageInfo* InInfo, UFoliageType* FoliageType, const TArray<FFoliageInstance>& Instances, const TSet<int32>& SelectedIndices) override;
+	virtual void PostEditUndo(FFoliageInfo* InInfo, UFoliageType* FoliageType) override;
 	virtual void PreMoveInstances(const TArray<int32>& InInstancesMoved) override;
 	virtual void PostMoveInstances(const TArray<int32>& InInstancesMoved, bool bFinished) override;
-	virtual void NotifyFoliageTypeChanged(UFoliageType* FoliageType, const TArray<FFoliageInstance>& Instances, const TSet<int32>& SelectedIndices, bool bSourceChanged) override;
-	void Reapply(const UFoliageType* FoliageType, const TArray<FFoliageInstance>& Instances, bool bPostLoad = false);
+	virtual void NotifyFoliageTypeChanged(UFoliageType* FoliageType, bool bSourceChanged) override;
+	virtual void Reapply(const UFoliageType* FoliageType) override;
 	AActor* Spawn(const FFoliageInstance& Instance);
 	TArray<AActor*> GetActorsFromSelectedIndices(const TSet<int32>& SelectedIndices) const;
 	virtual bool ShouldAttachToBaseComponent() const override { return bShouldAttachToBaseComponent; }
diff --git a/Engine/Source/Runtime/Foliage/Private/FoliageISMActor.cpp b/Engine/Source/Runtime/Foliage/Private/FoliageISMActor.cpp
new file mode 100644
--- /dev/null
+++ b/Engine/Source/Runtime/Foliage/Private/FoliageISMActor.cpp
@@ -0,0 +1,365 @@
+// Copyright Epic Games, Inc. All Rights Reserved.
+
+#include "FoliageISMActor.h"
+#include "InstancedFoliageActor.h"
+#include "FoliageType.h"
+#include "FoliageType_Actor.h"
+#include "FoliageHelper.h"
+#include "Engine/Engine.h"
+#include "Engine/World.h"
+#include "GameFramework/Actor.h"
+#include "Components/StaticMeshComponent.h"
+#include "ISMPartition/ISMComponentDescriptor.h"
+#include "FoliageInstancedStaticMeshComponent.h"
+#include "Math/Box.h"
+
+void FFoliageISMActor::Serialize(FArchive& Ar)
+{
+#if WITH_EDITORONLY_DATA
+	Ar << Guid;
+	ClientHandle.Serialize(Ar);
+	Ar << ISMDefinition;
+	Ar << ActorClass;
+#endif
+}
+
+#if WITH_EDITOR
+void FFoliageISMActor::AddReferencedObjects(UObject* InThis, FReferenceCollector& Collector)
+{
+	Collector.AddReferencedObject(ActorClass, InThis);
+}
+
+bool FFoliageISMActor::IsInitialized() const
+{
+	return ClientHandle.IsValid();
+}
+
+void InitDescriptorFromFoliageType(FISMComponentDescriptor& Descriptor, const UFoliageType* FoliageType)
+{
+	const UFoliageType_Actor* FoliageTypeActor = Cast<UFoliageType_Actor>(FoliageType);
+	Descriptor.ComponentClass = FoliageTypeActor->StaticMeshOnlyComponentClass != nullptr ? FoliageTypeActor->StaticMeshOnlyComponentClass.Get() : UFoliageInstancedStaticMeshComponent::StaticClass();
+	
+	Descriptor.Mobility = FoliageType->Mobility;
+	Descriptor.InstanceStartCullDistance = FoliageType->CullDistance.Min;
+	Descriptor.InstanceEndCullDistance = FoliageType->CullDistance.Max;
+	Descriptor.bCastShadow = FoliageType->CastShadow;
+	Descriptor.bCastDynamicShadow = FoliageType->bCastDynamicShadow;
+	Descriptor.bCastStaticShadow = FoliageType->bCastStaticShadow;
+	Descriptor.bCastContactShadow = FoliageType->bCastContactShadow;
+	Descriptor.RuntimeVirtualTextures = FoliageType->RuntimeVirtualTextures;
+	Descriptor.VirtualTextureRenderPassType = FoliageType->VirtualTextureRenderPassType;
+	Descriptor.VirtualTextureCullMips = FoliageType->VirtualTextureCullMips;
+	Descriptor.TranslucencySortPriority = FoliageType->TranslucencySortPriority;
+	Descriptor.bAffectDynamicIndirectLighting = FoliageType->bAffectDynamicIndirectLighting;
+	Descriptor.bAffectDistanceFieldLighting = FoliageType->bAffectDistanceFieldLighting;
+	Descriptor.bCastShadowAsTwoSided = FoliageType->bCastShadowAsTwoSided;
+	Descriptor.bReceivesDecals = FoliageType->bReceivesDecals;
+	Descriptor.bOverrideLightMapRes = FoliageType->bOverrideLightMapRes;
+	Descriptor.OverriddenLightMapRes = FoliageType->OverriddenLightMapRes;
+	Descriptor.LightmapType = FoliageType->LightmapType;
+	Descriptor.bUseAsOccluder = FoliageType->bUseAsOccluder;
+	Descriptor.bEnableDensityScaling = FoliageType->bEnableDensityScaling;
+	Descriptor.LightingChannels = FoliageType->LightingChannels;
+	Descriptor.bRenderCustomDepth = FoliageType->bRenderCustomDepth;
+	Descriptor.CustomDepthStencilWriteMask = FoliageType->CustomDepthStencilWriteMask;
+	Descriptor.CustomDepthStencilValue = FoliageType->CustomDepthStencilValue;
+	Descriptor.bIncludeInHLOD = FoliageType->bIncludeInHLOD;
+	Descriptor.BodyInstance.CopyBodyInstancePropertiesFrom(&FoliageType->BodyInstance);
+
+	Descriptor.bHasCustomNavigableGeometry = FoliageType->CustomNavigableGeometry;
+	Descriptor.bEnableDiscardOnLoad = FoliageType->bEnableDiscardOnLoad;
+}
+
+void FFoliageISMActor::Initialize(const UFoliageType* FoliageType)
+{
+	check(!IsInitialized());
+
+	AInstancedFoliageActor* IFA = GetIFA();
+
+	FoliageTypeActor = Cast<const UFoliageType_Actor>(FoliageType);
+	ActorClass = FoliageTypeActor->ActorClass ? FoliageTypeActor->ActorClass.Get() : AActor::StaticClass();
+	
+	FActorSpawnParameters SpawnParams;
+	SpawnParams.bCreateActorPackage = false;
+	SpawnParams.bNoFail = true;
+	SpawnParams.bHideFromSceneOutliner = true;
+	SpawnParams.bTemporaryEditorActor = true;
+	SpawnParams.ObjectFlags &= ~RF_Transactional;
+	SpawnParams.ObjectFlags |= RF_Transient;
+	AActor* SpawnedActor = IFA->GetWorld()->SpawnActor<AActor>(ActorClass, SpawnParams);
+	FTransform ActorTransform = SpawnedActor->GetActorTransform();
+	TArray<UStaticMeshComponent*> StaticMeshComponents;
+	SpawnedActor->GetComponents<UStaticMeshComponent>(StaticMeshComponents);
+
+	ClientHandle = IFA->RegisterClient(Guid);
+
+	for (UStaticMeshComponent* StaticMeshComponent : StaticMeshComponents)
+	{
+		FISMComponentDescriptor Descriptor;
+		Descriptor.InitFrom(StaticMeshComponent);
+		InitDescriptorFromFoliageType(Descriptor, FoliageTypeActor);
+		Descriptor.ComputeHash();
+
+		int32 DescriptorIndex = IFA->RegisterISMComponentDescriptor(Descriptor);
+		TArray<FTransform>& Transforms = ISMDefinition.FindOrAdd(DescriptorIndex);
+		if (UInstancedStaticMeshComponent* ISMComponent = Cast<UInstancedStaticMeshComponent>(StaticMeshComponent))
+		{
+			for (int32 InstanceIndex = 0; InstanceIndex < ISMComponent->GetInstanceCount(); ++InstanceIndex)
+			{
+				FTransform InstanceTransform;
+				if (ensure(ISMComponent->GetInstanceTransform(InstanceIndex, InstanceTransform, /*bWorldSpace=*/ true)))
+				{
+					FTransform LocalTransform = InstanceTransform.GetRelativeTransform(ActorTransform);
+					Transforms.Add(LocalTransform);
+				}
+			}
+		}
+		else
+		{
+			FTransform LocalTransform = StaticMeshComponent->GetComponentTransform().GetRelativeTransform(ActorTransform);
+			Transforms.Add(LocalTransform);
+		}
+	}
+
+	IFA->GetWorld()->DestroyActor(SpawnedActor);
+	
+	RegisterDelegates();
+}
+
+void FFoliageISMActor::Uninitialize()
+{
+	check(IsInitialized());
+	UnregisterDelegates();
+	GetIFA()->UnregisterClient(ClientHandle);
+	ISMDefinition.Empty();
+}
+
+void FFoliageISMActor::RegisterDelegates()
+{
+	if (ActorClass)
+	{
+		if (UBlueprint* Blueprint = Cast<UBlueprint>(ActorClass->ClassGeneratedBy))
+		{
+			Blueprint->OnCompiled().AddRaw(this, &FFoliageISMActor::OnBlueprintChanged);
+		}
+	}
+}
+
+void FFoliageISMActor::UnregisterDelegates()
+{
+	if (ActorClass)
+	{
+		if (UBlueprint* Blueprint = Cast<UBlueprint>(ActorClass->ClassGeneratedBy))
+		{
+			Blueprint->OnCompiled().RemoveAll(this);
+		}
+	}
+}
+
+void FFoliageISMActor::OnBlueprintChanged(UBlueprint* InBlueprint)
+{
+	Reapply(FoliageTypeActor);
+}
+
+void FFoliageISMActor::Reapply(const UFoliageType* FoliageType)
+{
+	if (IsInitialized())
+	{
+		Uninitialize();
+	}
+	Initialize(FoliageType);
+	check(IsInitialized());
+	
+	BeginUpdate();
+	for (const FFoliageInstance& Instance : Info->Instances)
+	{
+		AddInstance(Instance);
+	}
+	EndUpdate();
+}
+
+int32 FFoliageISMActor::GetInstanceCount() const
+{
+	return Info->Instances.Num();
+}
+
+void FFoliageISMActor::PreAddInstances(const UFoliageType* FoliageType, int32 AddedInstanceCount)
+{
+	if (!IsInitialized())
+	{
+		Initialize(FoliageType);
+		check(IsInitialized());
+	}
+
+	GetIFA()->ReserveISMInstances(ClientHandle, AddedInstanceCount, ISMDefinition);
+}
+
+void FFoliageISMActor::AddInstance(const FFoliageInstance& NewInstance)
+{
+	GetIFA()->AddISMInstance(ClientHandle, NewInstance.GetInstanceWorldTransform(), ISMDefinition);
+}
+
+void FFoliageISMActor::RemoveInstance(int32 InstanceIndex)
+{
+	GetIFA()->RemoveISMInstance(ClientHandle, InstanceIndex);
+}
+
+void FFoliageISMActor::BeginUpdate()
+{
+	GetIFA()->BeginUpdate();
+}
+
+void FFoliageISMActor::EndUpdate()
+{
+	GetIFA()->EndUpdate();
+}
+
+void FFoliageISMActor::SetInstanceWorldTransform(int32 InstanceIndex, const FTransform& Transform, bool bTeleport)
+{
+	GetIFA()->SetISMInstanceTransform(ClientHandle, InstanceIndex, Transform, bTeleport, ISMDefinition);
+}
+
+FTransform FFoliageISMActor::GetInstanceWorldTransform(int32 InstanceIndex) const
+{
+	return Info->Instances[InstanceIndex].GetInstanceWorldTransform();
+}
+
+bool FFoliageISMActor::IsOwnedComponent(const UPrimitiveComponent* Component) const
+{
+	return GetIFA()->IsISMComponent(Component);
+}
+
+void FFoliageISMActor::Refresh(bool bAsync, bool bForce)
+{
+	GetIFA()->UpdateHISMTrees(bAsync, bForce);
+}
+
+void FFoliageISMActor::OnHiddenEditorViewMaskChanged(uint64 InHiddenEditorViews)
+{
+	if (!IsInitialized())
+	{
+		return;
+	}
+
+	// This can give weird results if toggling the visibility of 2 foliage types that share the same meshes. The last one wins for now.
+	TArray<UInstancedStaticMeshComponent*> ClientComponents;
+	GetIFA()->GetClientComponents(ClientHandle, ClientComponents);
+	for (UInstancedStaticMeshComponent* Component : ClientComponents)
+	{
+		if (UFoliageInstancedStaticMeshComponent* FoliageComponent = Cast<UFoliageInstancedStaticMeshComponent>(Component))
+		{
+			FoliageComponent->FoliageHiddenEditorViews = InHiddenEditorViews;
+			FoliageComponent->MarkRenderStateDirty();
+		}
+	}
+}
+
+void FFoliageISMActor::PostEditUndo(FFoliageInfo* InInfo, UFoliageType* FoliageType)
+{
+	FFoliageImpl::PostEditUndo(InInfo, FoliageType);
+}
+
+void FFoliageISMActor::NotifyFoliageTypeWillChange(UFoliageType* FoliageType)
+{
+	UnregisterDelegates();
+}
+
+void FFoliageISMActor::NotifyFoliageTypeChanged(UFoliageType* FoliageType, bool bSourceChanged)
+{
+	if (!IsInitialized())
+	{
+		return;
+	}
+		
+	if (UFoliageType_Actor* InFoliageTypeActor = Cast<UFoliageType_Actor>(FoliageType))
+	{
+		// Implementation should change
+		if (!InFoliageTypeActor->bStaticMeshOnly)
+		{
+			Uninitialize();
+			return;
+		}
+	}
+
+	AInstancedFoliageActor* IFA = GetIFA();
+
+	// Go through descriptors and see if they changed
+	if (!bSourceChanged)
+	{
+		for (const auto& Pair : ISMDefinition)
+		{
+			const FISMComponentDescriptor& RegisteredDescriptor = IFA->GetISMComponentDescriptor(Pair.Key);
+			FISMComponentDescriptor NewDescriptor(RegisteredDescriptor);
+			InitDescriptorFromFoliageType(NewDescriptor, FoliageType);
+			NewDescriptor.ComputeHash();
+
+			if (RegisteredDescriptor != NewDescriptor)
+			{
+				bSourceChanged = true;
+				break;
+			}
+		}
+	}
+
+	if (bSourceChanged)
+	{
+		Reapply(FoliageType);
+		ApplySelection(true, Info->SelectedIndices);
+	}
+	else
+	{
+		RegisterDelegates();
+	}
+}
+
+void FFoliageISMActor::SelectAllInstances(bool bSelect)
+{
+	TSet<int32> Indices;
+	Indices.Reserve(Info->Instances.Num());
+	for (int32 i = 0; i < Info->Instances.Num(); ++i)
+	{
+		Indices.Add(i);
+	}
+	SelectInstances(bSelect, Indices);
+}
+
+void FFoliageISMActor::SelectInstance(bool bSelect, int32 Index)
+{
+	SelectInstances(bSelect, { Index });
+}
+
+void FFoliageISMActor::SelectInstances(bool bSelect, const TSet<int32>& SelectedIndices)
+{
+	GetIFA()->SelectISMInstances(ClientHandle, bSelect, SelectedIndices);
+}
+
+int32 FFoliageISMActor::GetInstanceIndexFrom(UInstancedStaticMeshComponent* HISMComponent, int32 ComponentIndex) const
+{
+	return GetIFA()->GetISMInstanceIndex(ClientHandle, HISMComponent, ComponentIndex);
+}
+
+FBox FFoliageISMActor::GetSelectionBoundingBox(const TSet<int32>& SelectedIndices) const
+{
+	return GetIFA()->GetISMInstanceBounds(ClientHandle, SelectedIndices);
+}
+void FFoliageISMActor::ApplySelection(bool bApply, const TSet<int32>& SelectedIndices)
+{
+	// Going in and out of Folaige with an Empty/Unregistered impl.
+	if (!IsInitialized())
+	{
+		return;
+	}
+
+	SelectAllInstances(false);
+	if (bApply)
+	{
+		SelectInstances(true, SelectedIndices);
+	}
+}
+
+void FFoliageISMActor::ClearSelection(const TSet<int32>& SelectedIndices)
+{
+	SelectAllInstances(false);
+}
+
+#endif // WITH_EDITOR
diff --git a/Engine/Source/Runtime/Foliage/Private/FoliageISMActor.h b/Engine/Source/Runtime/Foliage/Private/FoliageISMActor.h
new file mode 100644
--- /dev/null
+++ b/Engine/Source/Runtime/Foliage/Private/FoliageISMActor.h
@@ -0,0 +1,71 @@
+// Copyright Epic Games, Inc. All Rights Reserved.
+#pragma once
+
+#include "InstancedFoliage.h"
+#include "ISMPartition/ISMPartitionActor.h"
+#include "ISMPartition/ISMComponentDescriptor.h"
+#include "Misc/Guid.h"
+#include "Containers/SortedMap.h"
+
+struct FFoliageInfo;
+class UInstancedStaticMeshComponent;
+class UBlueprint;
+class UFoliageType_Actor;
+
+struct FFoliageISMActor : public FFoliageImpl
+{
+	FFoliageISMActor(FFoliageInfo* Info)
+		: FFoliageImpl(Info)
+#if WITH_EDITORONLY_DATA
+		, Guid(FGuid::NewGuid())
+		, ActorClass(nullptr)
+		, FoliageTypeActor(nullptr)
+#endif
+	{
+	}
+
+#if WITH_EDITORONLY_DATA
+	FGuid Guid;
+	FISMClientHandle ClientHandle;
+	TSortedMap<int32, TArray<FTransform>> ISMDefinition;
+	UClass* ActorClass;
+	const UFoliageType_Actor* FoliageTypeActor;
+#endif
+	virtual void Serialize(FArchive& Ar) override;
+		
+#if WITH_EDITOR
+	virtual void AddReferencedObjects(UObject* InThis, FReferenceCollector& Collector) override;
+	virtual bool IsInitialized() const override;
+	virtual void Initialize(const UFoliageType* FoliageType) override;
+	virtual void Uninitialize() override;
+	virtual void Reapply(const UFoliageType* FoliageType) override;
+	virtual int32 GetInstanceCount() const override;
+	virtual void PreAddInstances(const UFoliageType* FoliageType, int32 AddedInstanceCount) override;
+	virtual void AddInstance(const FFoliageInstance& NewInstance) override;
+	virtual void RemoveInstance(int32 InstanceIndex) override;
+	virtual void SetInstanceWorldTransform(int32 InstanceIndex, const FTransform& Transform, bool bTeleport) override;
+	virtual FTransform GetInstanceWorldTransform(int32 InstanceIndex) const override;
+	virtual bool IsOwnedComponent(const UPrimitiveComponent* Component) const override;
+	
+	virtual void SelectAllInstances(bool bSelect) override;
+	virtual void SelectInstance(bool bSelect, int32 Index) override;
+	virtual void SelectInstances(bool bSelect, const TSet<int32>& SelectedIndices) override;
+	virtual int32 GetInstanceIndexFrom(UInstancedStaticMeshComponent* HISMComponent, int32 ComponentIndex) const;
+	virtual FBox GetSelectionBoundingBox(const TSet<int32>& SelectedIndices) const override;
+	virtual void ApplySelection(bool bApply, const TSet<int32>& SelectedIndices) override;
+	virtual void ClearSelection(const TSet<int32>& SelectedIndices) override;
+
+	virtual void BeginUpdate() override;
+	virtual void EndUpdate() override;
+	virtual void Refresh(bool bAsync, bool bForce) override;
+	virtual void OnHiddenEditorViewMaskChanged(uint64 InHiddenEditorViews) override;
+	virtual void PostEditUndo(FFoliageInfo* InInfo, UFoliageType* FoliageType) override;
+	virtual void NotifyFoliageTypeWillChange(UFoliageType* FoliageType) override;
+	virtual void NotifyFoliageTypeChanged(UFoliageType* FoliageType, bool bSourceChanged) override;
+
+private:
+	void RegisterDelegates();
+	void UnregisterDelegates();
+	void OnBlueprintChanged(UBlueprint* InBlueprint);
+#endif
+};
\ No newline at end of file
diff --git a/Engine/Source/Runtime/Foliage/Private/InstancedFoliage.cpp b/Engine/Source/Runtime/Foliage/Private/InstancedFoliage.cpp
--- a/Engine/Source/Runtime/Foliage/Private/InstancedFoliage.cpp
+++ b/Engine/Source/Runtime/Foliage/Private/InstancedFoliage.cpp
@@ -46,6 +46,7 @@ InstancedFoliage.cpp: Instanced foliage implementation.
 #include "UObject/FortniteMainBranchObjectVersion.h"
 #include "PreviewScene.h"
 #include "FoliageActor.h"
+#include "FoliageISMActor.h"
 #include "LevelUtils.h"
 #include "FoliageHelper.h"
 #include "Algo/Transform.h"
@@ -105,6 +106,7 @@ struct FFoliageStaticMesh : public FFoliageImpl
 	virtual bool IsInitialized() const override { return Component != nullptr; }
 	virtual void Initialize(const UFoliageType* FoliageType) override;
 	virtual void Uninitialize() override;
+	virtual void Reapply(const UFoliageType* FoliageType) override;
 	virtual int32 GetInstanceCount() const override;
 	virtual void PreAddInstances(const UFoliageType* FoliageType, int32 Count) override;
 	virtual void AddInstance(const FFoliageInstance& NewInstance) override;
@@ -117,18 +119,19 @@ struct FFoliageStaticMesh : public FFoliageImpl
 	virtual void SelectAllInstances(bool bSelect) override;
 	virtual void SelectInstance(bool bSelect, int32 Index) override;
 	virtual void SelectInstances(bool bSelect, const TSet<int32>& SelectedIndices) override;
+	virtual int32 GetInstanceIndexFrom(UInstancedStaticMeshComponent* HISMComponent, int32 ComponentIndex) const override;
 	virtual FBox GetSelectionBoundingBox(const TSet<int32>& SelectedIndices) const override;
 	virtual void ApplySelection(bool bApply, const TSet<int32>& SelectedIndices) override;
 	virtual void ClearSelection(const TSet<int32>& SelectedIndices) override;
 
 	virtual void BeginUpdate() override;
 	virtual void EndUpdate() override;
-	virtual void Refresh(const TArray<FFoliageInstance>& Instances, bool Async, bool Force) override;
+	virtual void Refresh(bool Async, bool Force) override;
 	virtual void OnHiddenEditorViewMaskChanged(uint64 InHiddenEditorViews) override;
 	virtual void PreEditUndo(UFoliageType* FoliageType) override;
-	virtual void PostEditUndo(FFoliageInfo* InInfo, UFoliageType* FoliageType, const TArray<FFoliageInstance>& Instances, const TSet<int32>& SelectedIndices) override;
+	virtual void PostEditUndo(FFoliageInfo* InInfo, UFoliageType* FoliageType) override;
 	virtual void NotifyFoliageTypeWillChange(UFoliageType* FoliageType) override;
-	virtual void NotifyFoliageTypeChanged(UFoliageType* FoliageType, const TArray<FFoliageInstance>& Instances, const TSet<int32>& SelectedIndices, bool bSourceChanged) override;
+	virtual void NotifyFoliageTypeChanged(UFoliageType* FoliageType, bool bSourceChanged) override;
 	virtual void EnterEditMode() override;
 	virtual void ExitEditMode() override;
 
@@ -142,8 +145,7 @@ struct FFoliageStaticMesh : public FFoliageImpl
 
 	void UpdateComponentSettings(const UFoliageType_InstancedStaticMesh* InSettings);
 	// Recreate the component if the FoliageType's ComponentClass doesn't match the Component's class
-	void CheckComponentClass(const UFoliageType_InstancedStaticMesh* InSettings, const TArray<FFoliageInstance>& Instances, const TSet<int32>& SelectedIndices);
-	void ReapplyInstancesToComponent(const TArray<FFoliageInstance>& Instances, const TSet<int32>& SelectedIndices);
+	void CheckComponentClass(const UFoliageType_InstancedStaticMesh* InSettings);
 	void CreateNewComponent(const UFoliageType* InSettings);
 };
 
@@ -765,6 +767,7 @@ UFoliageType_Actor::UFoliageType_Actor(const FObjectInitializer& ObjectInitializ
 	Density = 10;
 	Radius = 500;
 	bShouldAttachToBaseComponent = true;
+	StaticMeshOnlyComponentClass = UFoliageInstancedStaticMeshComponent::StaticClass();
 }
 
 #if WITH_EDITOR
@@ -1109,6 +1112,16 @@ void FFoliageStaticMesh::SelectInstances(bool bSelect, const TSet<int32>& Select
 	Component->MarkRenderStateDirty();
 }
 
+int32 FFoliageStaticMesh::GetInstanceIndexFrom(UInstancedStaticMeshComponent* HISMComponent, int32 ComponentIndex) const
+{
+	if (IsOwnedComponent(HISMComponent))
+	{
+		return ComponentIndex;
+	}
+
+	return INDEX_NONE;
+}
+
 FBox FFoliageStaticMesh::GetSelectionBoundingBox(const TSet<int32>& SelectedIndices) const
 {
 	FBox BoundingBox(EForceInit::ForceInit);
@@ -1173,7 +1186,7 @@ void FFoliageStaticMesh::EndUpdate()
 	}
 }
 
-void FFoliageStaticMesh::Refresh(const TArray<FFoliageInstance>& Instances, bool Async, bool Force)
+void FFoliageStaticMesh::Refresh(bool Async, bool Force)
 {
 	if (Component != nullptr)
 	{
@@ -1203,9 +1216,9 @@ void FFoliageStaticMesh::PreEditUndo(UFoliageType* FoliageType)
 	}
 }
 
-void FFoliageStaticMesh::PostEditUndo(FFoliageInfo* InInfo, UFoliageType* FoliageType, const TArray<FFoliageInstance>& Instances, const TSet<int32>& SelectedIndices)
+void FFoliageStaticMesh::PostEditUndo(FFoliageInfo* InInfo, UFoliageType* FoliageType)
 {
-	FFoliageImpl::PostEditUndo(InInfo, FoliageType, Instances, SelectedIndices);
+	FFoliageImpl::PostEditUndo(InInfo, FoliageType);
 	if (UFoliageType_InstancedStaticMesh* FoliageType_InstancedStaticMesh = Cast<UFoliageType_InstancedStaticMesh>(FoliageType))
 	{
 		if (Component != nullptr && FoliageType_InstancedStaticMesh->GetStaticMesh() != nullptr)
@@ -1213,8 +1226,8 @@ void FFoliageStaticMesh::PostEditUndo(FFoliageInfo* InInfo, UFoliageType* Foliag
 			FoliageType_InstancedStaticMesh->GetStaticMesh()->GetOnExtendedBoundsChanged().AddRaw(this, &FFoliageStaticMesh::HandleComponentMeshBoundsChanged);
 		}
 
-		CheckComponentClass(FoliageType_InstancedStaticMesh, Instances, SelectedIndices);
-		ReapplyInstancesToComponent(Instances, SelectedIndices);
+		CheckComponentClass(FoliageType_InstancedStaticMesh);
+		Reapply(FoliageType);
 	}
 }
 
@@ -1232,11 +1245,11 @@ void FFoliageStaticMesh::NotifyFoliageTypeWillChange(UFoliageType* FoliageType)
 	}
 }
 
-void FFoliageStaticMesh::NotifyFoliageTypeChanged(UFoliageType* FoliageType, const TArray<FFoliageInstance>& Instances, const TSet<int32>& SelectedIndices, bool bSourceChanged)
+void FFoliageStaticMesh::NotifyFoliageTypeChanged(UFoliageType* FoliageType, bool bSourceChanged)
 {
 	UFoliageType_InstancedStaticMesh* FoliageType_InstancedStaticMesh = Cast<UFoliageType_InstancedStaticMesh>(FoliageType);
 	check(FoliageType_InstancedStaticMesh);
-	CheckComponentClass(FoliageType_InstancedStaticMesh, Instances, SelectedIndices);
+	CheckComponentClass(FoliageType_InstancedStaticMesh);
 	UpdateComponentSettings(FoliageType_InstancedStaticMesh);
 	
 	if (bSourceChanged && Component != nullptr && Component->GetStaticMesh() != nullptr)
@@ -1342,7 +1355,7 @@ void FFoliageStaticMesh::HandleComponentMeshBoundsChanged(const FBoxSphereBounds
 	}
 }
 
-void FFoliageStaticMesh::CheckComponentClass(const UFoliageType_InstancedStaticMesh* InSettings, const TArray<FFoliageInstance>& Instances, const TSet<int32>& SelectedIndices)
+void FFoliageStaticMesh::CheckComponentClass(const UFoliageType_InstancedStaticMesh* InSettings)
 {
 	if (Component)
 	{
@@ -1364,7 +1377,7 @@ void FFoliageStaticMesh::CheckComponentClass(const UFoliageType_InstancedStaticM
 			Initialize(InSettings);
 
 			// apply the instances to it
-			ReapplyInstancesToComponent(Instances, SelectedIndices);
+			Reapply(InSettings);
 		}
 	}
 }
@@ -1587,7 +1600,7 @@ void FFoliageStaticMesh::UpdateComponentSettings(const UFoliageType_InstancedSta
 	}
 }
 
-void FFoliageStaticMesh::ReapplyInstancesToComponent(const TArray<FFoliageInstance>& Instances, const TSet<int32>& SelectedIndices)
+void FFoliageStaticMesh::Reapply(const UFoliageType* FoliageType)
 {
 	if (Component)
 	{
@@ -1601,7 +1614,7 @@ void FFoliageStaticMesh::ReapplyInstancesToComponent(const TArray<FFoliageInstan
 
 		Component->bAutoRebuildTreeOnInstanceChanges = false;
 
-		for (auto& Instance : Instances)
+		for (auto& Instance : Info->Instances)
 		{
 			Component->AddInstanceWorldSpace(Instance.GetInstanceWorldTransform());
 		}
@@ -1611,9 +1624,9 @@ void FFoliageStaticMesh::ReapplyInstancesToComponent(const TArray<FFoliageInstan
 
 		Component->ClearInstanceSelection();
 
-		if (SelectedIndices.Num())
+		if (Info->SelectedIndices.Num())
 		{
-			for (int32 i : SelectedIndices)
+			for (int32 i : Info->SelectedIndices)
 			{
 				Component->SelectInstance(true, i, 1);
 			}
@@ -1682,6 +1695,10 @@ void FFoliageInfo::CreateImplementation(EFoliageImplType InType)
 	{
 		Implementation.Reset(new FFoliageActor(this));
 	}
+	else if (Type == EFoliageImplType::ISMActor)
+	{
+		Implementation.Reset(new FFoliageISMActor(this));
+	}
 }
 
 int32 FFoliageInfo::GetOverlappingSphereCount(const FSphere& Sphere) const
@@ -1721,20 +1738,33 @@ void FFoliageInfo::GetOverlappingMeshCount(const FSphere& Sphere, TMap<UStaticMe
 }
 
 #if WITH_EDITOR
-void FFoliageInfo::CreateImplementation(const UFoliageType* FoliageType)
+EFoliageImplType FFoliageInfo::GetImplementationType(const UFoliageType* FoliageType) const
 {
-	check(!Implementation.IsValid());
-	// Change Impl based on FoliageType param
-	EFoliageImplType ImplType = EFoliageImplType::Unknown;
 	if (FoliageType->IsA<UFoliageType_InstancedStaticMesh>())
 	{
-		ImplType = EFoliageImplType::StaticMesh;
+		return EFoliageImplType::StaticMesh;
 	}
 	else if (FoliageType->IsA<UFoliageType_Actor>())
 	{
-		ImplType = EFoliageImplType::Actor;
+		const UFoliageType_Actor* ActorFoliageType = Cast<UFoliageType_Actor>(FoliageType);
+		if (ActorFoliageType->bStaticMeshOnly)
+		{
+			return EFoliageImplType::ISMActor;
+		}
+		else
+		{
+			return EFoliageImplType::Actor;
+		}
 	}
-	CreateImplementation(ImplType);
+
+	return EFoliageImplType::Unknown;
+}
+
+void FFoliageInfo::CreateImplementation(const UFoliageType* FoliageType)
+{
+	check(!Implementation.IsValid());
+	
+	CreateImplementation(GetImplementationType(FoliageType));
 }
 
 void FFoliageInfo::Initialize(const UFoliageType* FoliageType)
@@ -1747,6 +1777,7 @@ void FFoliageInfo::Initialize(const UFoliageType* FoliageType)
 void FFoliageInfo::Uninitialize()
 {
 	check(IsInitialized());
+	FoliageTypeUpdateGuid.Invalidate();
 	Implementation->Uninitialize();
 }
 
@@ -1762,7 +1793,15 @@ void FFoliageInfo::NotifyFoliageTypeWillChange(UFoliageType* FoliageType)
 
 void FFoliageInfo::NotifyFoliageTypeChanged(UFoliageType* FoliageType, bool bSourceChanged)
 {
-	Implementation->NotifyFoliageTypeChanged(FoliageType, Instances, SelectedIndices, bSourceChanged);
+	FoliageTypeUpdateGuid = FoliageType->UpdateGuid;
+	// Handle Implementation being uninitialized by FoliageType change
+	bool bWasInitialized = Implementation->IsInitialized();
+	
+	Implementation->NotifyFoliageTypeChanged(FoliageType, bSourceChanged);
+	if (bWasInitialized && !Implementation->IsInitialized())
+	{
+		ReallocateClusters(FoliageType);
+	}
 }
 
 void FFoliageInfo::CheckValid()
@@ -2194,6 +2233,7 @@ void FFoliageInfo::ReallocateClusters(UFoliageType* InSettings)
 	// Copy the UpdateGuid from the foliage type
 	FoliageTypeUpdateGuid = InSettings->UpdateGuid;
 
+	Implementation->BeginUpdate();
 	// Re-add
 	for (FFoliageInstance& Instance : OldInstances)
 	{
@@ -2202,7 +2242,7 @@ void FFoliageInfo::ReallocateClusters(UFoliageType* InSettings)
 			AddInstance(InSettings, Instance);
 		}
 	}
-
+	Implementation->EndUpdate();
 	
 	Refresh(true, true);
 }
@@ -2347,7 +2387,7 @@ void FFoliageInfo::SelectInstances(bool bSelect, TArray<int32>& InInstances)
 void FFoliageInfo::Refresh(bool Async, bool Force)
 {
 	check(Implementation.IsValid())
-	Implementation->Refresh(this->Instances, Async, Force);
+	Implementation->Refresh(Async, Force);
 }
 
 void FFoliageInfo::OnHiddenEditorViewMaskChanged(uint64 InHiddenEditorViews)
@@ -2364,7 +2404,7 @@ void FFoliageInfo::PostEditUndo(AInstancedFoliageActor* InIFA, UFoliageType* Fol
 {
 	// Set the IFA after PostEditUndo as it is not a serialized member (will be nulled on serialization)
 	IFA = InIFA;
-	Implementation->PostEditUndo(this, FoliageType, Instances, SelectedIndices);
+	Implementation->PostEditUndo(this, FoliageType);
 		
 	// Regenerate instance hash
 	// We regenerate it here instead of saving to transaction buffer to speed up modify operations
@@ -2398,9 +2438,9 @@ void FFoliageInfo::RemoveBaseComponentOnInstances()
 
 void FFoliageInfo::IncludeActor(const UFoliageType* FoliageType, AActor* InActor)
 {
-	if (FFoliageActor* FoliageActor = StaticCast<FFoliageActor*>(Implementation.Get()))
+	if (Type == EFoliageImplType::Actor)
 	{
-		if (FoliageActor->FindIndex(InActor) == INDEX_NONE)
+		if (FFoliageActor* FoliageActor = StaticCast<FFoliageActor*>(Implementation.Get()))
 		{
 			FFoliageInstance NewInstance;
 			NewInstance.BaseComponent = nullptr;
@@ -2414,31 +2454,38 @@ void FFoliageInfo::IncludeActor(const UFoliageType* FoliageType, AActor* InActor
 
 			int32 Index = Instances.Add(NewInstance);
 			InstanceHash->InsertInstance(NewInstance.Location, Index);
-			FoliageActor->PreAddInstances(FoliageType, 1);
-			FoliageActor->ActorInstances.Add(InActor);
-			InActor->Modify();
-			FFoliageHelper::SetIsOwnedByFoliage(InActor);
+
+			if (FoliageActor->FindIndex(InActor) == INDEX_NONE)
+			{
+				FoliageActor->PreAddInstances(FoliageType, 1);
+				FoliageActor->ActorInstances.Add(InActor);
+				InActor->Modify();
+				FFoliageHelper::SetIsOwnedByFoliage(InActor);
+			}
 		}
 	}
 }
 
 void FFoliageInfo::ExcludeActors()
 {
-	if (FFoliageActor* FoliageActor = StaticCast<FFoliageActor*>(Implementation.Get()))
+	if (Type == EFoliageImplType::Actor)
 	{
-		SelectedIndices.Empty();
-		Instances.Empty();
-		InstanceHash->Empty();
-		ComponentHash.Empty();
-		for (AActor* Actor : FoliageActor->ActorInstances)
+		if (FFoliageActor* FoliageActor = StaticCast<FFoliageActor*>(Implementation.Get()))
 		{
-			if (Actor != nullptr)
+			SelectedIndices.Empty();
+			Instances.Empty();
+			InstanceHash->Empty();
+			ComponentHash.Empty();
+			for (AActor* Actor : FoliageActor->ActorInstances)
 			{
-				Actor->Modify();
-				FFoliageHelper::SetIsOwnedByFoliage(Actor, false);
+				if (Actor != nullptr)
+				{
+					Actor->Modify();
+					FFoliageHelper::SetIsOwnedByFoliage(Actor, false);
+				}
 			}
+			FoliageActor->ActorInstances.Empty();
 		}
-		FoliageActor->ActorInstances.Empty();
 	}
 }
 
@@ -2552,6 +2599,7 @@ TUniqueObj<FFoliageInfo>& AInstancedFoliageActor::AddFoliageInfo(UFoliageType* F
 	TUniqueObj<FFoliageInfo>& NewFoliageInfo = FoliageInfos.Add(FoliageType);
 #if WITH_EDITORONLY_DATA
 	NewFoliageInfo->IFA = this;
+	NewFoliageInfo->FoliageTypeUpdateGuid = FoliageType->UpdateGuid;
 #endif
 	return NewFoliageInfo;
 }
@@ -2561,6 +2609,7 @@ TUniqueObj<FFoliageInfo>& AInstancedFoliageActor::AddFoliageInfo(UFoliageType* F
 	TUniqueObj<FFoliageInfo>& NewFoliageInfo = FoliageInfos.Add(FoliageType, MoveTemp(FoliageInfo));
 #if WITH_EDITORONLY_DATA
 	NewFoliageInfo->IFA = this;
+	NewFoliageInfo->FoliageTypeUpdateGuid = FoliageType->UpdateGuid;
 #endif
 	return NewFoliageInfo;
 }
@@ -2715,7 +2764,7 @@ void AInstancedFoliageActor::MoveInstancesForMovedOwnedActors(AActor* InActor)
 			continue;
 		}
 
-		if (Pair.Key->IsA<UFoliageType_Actor>())
+		if (Pair.Key->IsA<UFoliageType_Actor>() && Pair.Value->Type == EFoliageImplType::Actor)
 		{
 			if (FFoliageActor* FoliageActor = StaticCast<FFoliageActor*>(Pair.Value->Implementation.Get()))
 			{
@@ -2866,6 +2915,7 @@ void AInstancedFoliageActor::DeleteInstancesForComponent(UWorld* InWorld, UActor
 void AInstancedFoliageActor::DeleteInstancesForProceduralFoliageComponent(const UProceduralFoliageComponent* ProceduralFoliageComponent, bool InRebuildTree)
 {
 	const FGuid& ProceduralGuid = ProceduralFoliageComponent->GetProceduralGuid();
+	BeginUpdate();
 	for (auto& Pair : FoliageInfos)
 	{
 		FFoliageInfo& Info = *Pair.Value;
@@ -2883,7 +2933,7 @@ void AInstancedFoliageActor::DeleteInstancesForProceduralFoliageComponent(const
 			Info.RemoveInstances(InstancesToRemove, InRebuildTree);
 		}
 	}
-
+	EndUpdate();
 	// Clean up dead cross-level references
 	FFoliageInstanceBaseCache::CompactInstanceBaseCache(this);
 }
@@ -3340,7 +3390,7 @@ void AInstancedFoliageActor::ClearSelection()
 	}
 }
 
-void AInstancedFoliageActor::SelectInstance(UInstancedStaticMeshComponent* InComponent, int32 InInstanceIndex, bool bToggle)
+void AInstancedFoliageActor::SelectInstance(UInstancedStaticMeshComponent* InComponent, int32 InComponentInstanceIndex, bool bToggle)
 {
 	Modify(false);
 
@@ -3354,44 +3404,37 @@ void AInstancedFoliageActor::SelectInstance(UInstancedStaticMeshComponent* InCom
 	{
 		UFoliageType* Type = nullptr;
 		FFoliageInfo* Info = nullptr;
+		int32 InstanceIndex = INDEX_NONE;
 
 		for (auto& Pair : FoliageInfos)
 		{
-			if (Pair.Value->Type == EFoliageImplType::StaticMesh)
+			InstanceIndex = Pair.Value->Implementation->GetInstanceIndexFrom(InComponent, InComponentInstanceIndex);
+			if (InstanceIndex != INDEX_NONE)
 			{
-				FFoliageStaticMesh* FoliageStaticMesh = StaticCast<FFoliageStaticMesh*>(Pair.Value->Implementation.Get());
-				if (FoliageStaticMesh->Component == InComponent)
-				{
-					Type = Pair.Key;
-					Info = &Pair.Value.Get();
-					break;
-				}
+				Type = Pair.Key;
+				Info = &Pair.Value.Get();
+				break;
 			}
 		}
 
 		if (Info)
 		{
-			bool bIsSelected = Info->SelectedIndices.Contains(InInstanceIndex);
+			bool bIsSelected = Info->SelectedIndices.Contains(InstanceIndex);
 
 			// Deselect if it's already selected.
-			if (InInstanceIndex < InComponent->SelectedInstances.Num())
-			{
-				InComponent->SelectInstance(false, InInstanceIndex, 1);
-				InComponent->MarkRenderStateDirty();
-			}
+			Info->Implementation->SelectInstance(false, InstanceIndex);
 
 			if (bIsSelected)
 			{
-				Info->SelectedIndices.Remove(InInstanceIndex);
+				Info->SelectedIndices.Remove(InstanceIndex);
 			}
 
 			if (!bToggle || !bIsSelected)
 			{
 				// Add the selection
-				InComponent->SelectInstance(true, InInstanceIndex, 1);
-				InComponent->MarkRenderStateDirty();
+				Info->Implementation->SelectInstance(true, InstanceIndex);
 
-				Info->SelectedIndices.Add(InInstanceIndex);
+				Info->SelectedIndices.Add(InstanceIndex);
 			}
 		}
 	}
@@ -3572,12 +3615,15 @@ void AInstancedFoliageActor::PostEditUndo()
 
 	InstanceBaseCache.UpdateInstanceBaseCachedTransforms();
 
+	BeginUpdate();
+
 	for (auto& Pair : FoliageInfos)
 	{
 		FFoliageInfo& Info = *Pair.Value;
 
 		Info.PostEditUndo(this, Pair.Key);
 	}
+	EndUpdate();
 }
 
 void AInstancedFoliageActor::PostDuplicate(bool bDuplicateForPIE)
@@ -4055,7 +4101,7 @@ void AInstancedFoliageActor::PostLoad()
 				{
 					FFoliageStaticMesh* FoliageStaticMesh = StaticCast<FFoliageStaticMesh*>(Info.Implementation.Get());
 					UFoliageType_InstancedStaticMesh* FoliageType_InstancedStaticMesh = Cast<UFoliageType_InstancedStaticMesh>(FoliageType);
-					FoliageStaticMesh->CheckComponentClass(FoliageType_InstancedStaticMesh, Info.Instances, Info.SelectedIndices);
+					FoliageStaticMesh->CheckComponentClass(FoliageType_InstancedStaticMesh);
 				}
 			}
 
@@ -4231,6 +4277,7 @@ void AInstancedFoliageActor::PostLoad()
 
 	if (!GIsEditor && CVarFoliageDiscardDataOnLoad.GetValueOnGameThread())
 	{
+		bool bHasISMFoliage = false;
 		for (auto& Pair : FoliageInfos)
 		{
 			if (!Pair.Key || Pair.Key->bEnableDiscardOnLoad)
@@ -4250,10 +4297,28 @@ void AInstancedFoliageActor::PostLoad()
 					FFoliageActor* FoliageActor = StaticCast<FFoliageActor*>(Pair.Value->Implementation.Get());
 					FoliageActor->DestroyActors(true);
 				}
+				else if (Pair.Value->Type == EFoliageImplType::ISMActor)
+				{
+					bHasISMFoliage = true;
+				}
 			}
 				
 			Pair.Value = FFoliageInfo();
 		}
+
+		if (bHasISMFoliage)
+		{
+			TArray<UFoliageInstancedStaticMeshComponent*> FoliageComponents;
+			GetComponents<UFoliageInstancedStaticMeshComponent>(FoliageComponents);
+			for (UFoliageInstancedStaticMeshComponent* FoliageComponent : FoliageComponents)
+			{
+				if (FoliageComponent && FoliageComponent->bEnableDiscardOnLoad)
+				{
+					FoliageComponent->ConditionalPostLoad();
+					FoliageComponent->DestroyComponent();
+				}
+			}
+		}
 	}
 }
 
@@ -4529,19 +4594,30 @@ void AInstancedFoliageActor::DetectFoliageTypeChangeAndUpdate()
 				{
 					FFoliageStaticMesh* FoliageStaticMesh = StaticCast<FFoliageStaticMesh*>(Info.Implementation.Get());
 					UFoliageType_InstancedStaticMesh* FoliageType_InstancedStaticMesh = Cast<UFoliageType_InstancedStaticMesh>(FoliageType);
-					FoliageStaticMesh->CheckComponentClass(FoliageType_InstancedStaticMesh, Info.Instances, Info.SelectedIndices);
+					FoliageStaticMesh->CheckComponentClass(FoliageType_InstancedStaticMesh);
 					FoliageStaticMesh->UpdateComponentSettings(FoliageType_InstancedStaticMesh);
 				}
-				// Respawn foliage actors
-				else if (Info.Type == EFoliageImplType::Actor)
+				// Respawn foliage 
+				else
 				{
 					// We can't spawn in postload because BeingPlay might call UnrealScript which is not supported.
 					UWorld* World = GetWorld();
 					if (World && !World->IsGameWorld())
 					{
-						FFoliageActor* FoliageActor = StaticCast<FFoliageActor*>(Info.Implementation.Get());
-						const bool bPostLoad = true;
-						FoliageActor->Reapply(FoliageType, Info.Instances, bPostLoad);
+						Info.Implementation->NotifyFoliageTypeWillChange(FoliageType);
+
+						EFoliageImplType CurrentType = Info.GetImplementationType(FoliageType);
+						if (Info.Type != Info.GetImplementationType(FoliageType))
+						{
+							Info.Implementation->Uninitialize();
+							Info.Implementation.Reset();
+							Info.CreateImplementation(CurrentType);
+							Info.Implementation->Reapply(FoliageType);
+						}
+						else
+						{
+							Info.Implementation->NotifyFoliageTypeChanged(FoliageType, false);
+						}
 					}
 				}
 			}
@@ -4938,6 +5014,7 @@ float AInstancedFoliageActor::InternalTakeRadialDamage(float Damage, struct FRad
 
 UFoliageInstancedStaticMeshComponent::UFoliageInstancedStaticMeshComponent(const FObjectInitializer& ObjectInitializer)
 	: Super(ObjectInitializer)
+	, bEnableDiscardOnLoad(false)
 {
 #if WITH_EDITORONLY_DATA
 	bEnableAutoLODGeneration = false;
diff --git a/Engine/Source/Runtime/Foliage/Public/FoliageInstancedStaticMeshComponent.h b/Engine/Source/Runtime/Foliage/Public/FoliageInstancedStaticMeshComponent.h
--- a/Engine/Source/Runtime/Foliage/Public/FoliageInstancedStaticMeshComponent.h
+++ b/Engine/Source/Runtime/Foliage/Public/FoliageInstancedStaticMeshComponent.h
@@ -23,14 +23,17 @@ class FOLIAGE_API UFoliageInstancedStaticMeshComponent : public UHierarchicalIns
 public:
 
 	UPROPERTY(BlueprintAssignable, Category = "Game|Damage")
-		FInstancePointDamageSignature OnInstanceTakePointDamage;
+	FInstancePointDamageSignature OnInstanceTakePointDamage;
 
 	UPROPERTY(BlueprintAssignable, Category = "Game|Damage")
-		FInstanceRadialDamageSignature OnInstanceTakeRadialDamage;
+	FInstanceRadialDamageSignature OnInstanceTakeRadialDamage;
+
+	UPROPERTY()
+	bool bEnableDiscardOnLoad;
 
 #if WITH_EDITORONLY_DATA
 	UPROPERTY(transient)
-		uint64 FoliageHiddenEditorViews;
+	uint64 FoliageHiddenEditorViews;
 #endif// WITH_EDITOR
 
 	virtual void ReceiveComponentDamage(float DamageAmount, FDamageEvent const& DamageEvent, AController* EventInstigator, AActor* DamageCauser) override;
diff --git a/Engine/Source/Runtime/Foliage/Public/FoliageType_Actor.h b/Engine/Source/Runtime/Foliage/Public/FoliageType_Actor.h
--- a/Engine/Source/Runtime/Foliage/Public/FoliageType_Actor.h
+++ b/Engine/Source/Runtime/Foliage/Public/FoliageType_Actor.h
@@ -5,6 +5,8 @@
 #include "FoliageType.h"
 #include "FoliageType_Actor.generated.h"
 
+class UFoliageInstancedStaticMeshComponent;
+
 UCLASS(hidecategories = Object, editinlinenew, MinimalAPI)
 class UFoliageType_Actor : public UFoliageType
 {
@@ -16,6 +18,12 @@ class UFoliageType_Actor : public UFoliageType
 	UPROPERTY(EditAnywhere, Category = Actor)
 	bool bShouldAttachToBaseComponent;
 
+	UPROPERTY(EditAnywhere, Category = Mesh, meta = (ToolTip = "If enabled, will place an instanced static mesh representation of this actor without placing an actual actor"))
+	bool bStaticMeshOnly;
+		
+	UPROPERTY(EditAnywhere, Category = Mesh, meta = (EditCondition = "bStaticMeshOnly"))
+	TSubclassOf<UFoliageInstancedStaticMeshComponent> StaticMeshOnlyComponentClass;
+
 	virtual UObject* GetSource() const override { return ActorClass; }
 
 #if WITH_EDITOR
diff --git a/Engine/Source/Runtime/Foliage/Public/InstancedFoliage.h b/Engine/Source/Runtime/Foliage/Public/InstancedFoliage.h
--- a/Engine/Source/Runtime/Foliage/Public/InstancedFoliage.h
+++ b/Engine/Source/Runtime/Foliage/Public/InstancedFoliage.h
@@ -13,6 +13,7 @@ InstancedFoliage.h: Instanced foliage type definitions.
 class AInstancedFoliageActor;
 class UActorComponent;
 class UFoliageType;
+class UInstancedStaticMeshComponent;
 class UHierarchicalInstancedStaticMeshComponent;
 class UFoliageType_InstancedStaticMesh;
 class UPrimitiveComponent;
@@ -165,7 +166,8 @@ enum class EFoliageImplType : uint8
 {
 	Unknown = 0,
 	StaticMesh = 1,
-	Actor = 2
+	Actor = 2,
+	ISMActor = 3
 };
 
 struct FFoliageInfo;
@@ -194,6 +196,7 @@ struct FFoliageImpl
 	virtual bool IsInitialized() const = 0;
 	virtual void Initialize(const UFoliageType* FoliageType) = 0;
 	virtual void Uninitialize() = 0;
+	virtual void Reapply(const UFoliageType* FoliageType) = 0;
 	virtual int32 GetInstanceCount() const = 0;
 	virtual void PreAddInstances(const UFoliageType* FoliageType, int32 Count) = 0;
 	virtual void AddInstance(const FFoliageInstance& NewInstance) = 0;
@@ -211,18 +214,19 @@ struct FFoliageImpl
 	virtual void SelectAllInstances(bool bSelect) = 0;
 	virtual void SelectInstance(bool bSelect, int32 Index) = 0;
 	virtual void SelectInstances(bool bSelect, const TSet<int32>& SelectedIndices) = 0;
+	virtual int32 GetInstanceIndexFrom(UInstancedStaticMeshComponent* HISMComponent, int32 ComponentIndex) const { return INDEX_NONE; }
 	virtual FBox GetSelectionBoundingBox(const TSet<int32>& SelectedIndices) const = 0;
 	virtual void ApplySelection(bool bApply, const TSet<int32>& SelectedIndices) = 0;
 	virtual void ClearSelection(const TSet<int32>& SelectedIndices) = 0;
 
 	virtual void BeginUpdate() {}
 	virtual void EndUpdate() {}
-	virtual void Refresh(const TArray<FFoliageInstance>& Instances, bool Async, bool Force) {}
+	virtual void Refresh(bool Async, bool Force) {}
 	virtual void OnHiddenEditorViewMaskChanged(uint64 InHiddenEditorViews) = 0;
 	virtual void PreEditUndo(UFoliageType* FoliageType) {}
-	virtual void PostEditUndo(FFoliageInfo* InInfo, UFoliageType* FoliageType, const TArray<FFoliageInstance>& Instances, const TSet<int32>& SelectedIndices) { Info = InInfo; }
+	virtual void PostEditUndo(FFoliageInfo* InInfo, UFoliageType* FoliageType) { Info = InInfo; }
 	virtual void NotifyFoliageTypeWillChange(UFoliageType* FoliageType) {}
-	virtual void NotifyFoliageTypeChanged(UFoliageType* FoliageType, const TArray<FFoliageInstance>& Instances, const TSet<int32>& SelectedIndices, bool bSourceChanged) = 0;
+	virtual void NotifyFoliageTypeChanged(UFoliageType* FoliageType, bool bSourceChanged) = 0;
 	virtual void EnterEditMode() {}
 	virtual void ExitEditMode() {}
 	virtual bool ShouldAttachToBaseComponent() const { return true; }
@@ -291,6 +295,7 @@ struct FFoliageInfo
 	
 #if WITH_EDITOR
 	
+	FOLIAGE_API EFoliageImplType GetImplementationType(const UFoliageType* FoliageType) const;
 	FOLIAGE_API void CreateImplementation(const UFoliageType* FoliageType);
 	FOLIAGE_API void NotifyFoliageTypeWillChange(UFoliageType* FoliageType);
 	FOLIAGE_API void NotifyFoliageTypeChanged(UFoliageType* FoliageType, bool bSourceChanged);
diff --git a/Engine/Source/Runtime/Foliage/Public/InstancedFoliageActor.h b/Engine/Source/Runtime/Foliage/Public/InstancedFoliageActor.h
--- a/Engine/Source/Runtime/Foliage/Public/InstancedFoliageActor.h
+++ b/Engine/Source/Runtime/Foliage/Public/InstancedFoliageActor.h
@@ -11,7 +11,7 @@
 #include "FoliageInstanceBase.h"
 #include "InstancedFoliage.h"
 #include "InstancedFoliageCustomVersion.h"
-#include "ActorPartition/PartitionActor.h"
+#include "ISMPartition/ISMPartitionActor.h"
 
 #include "InstancedFoliageActor.generated.h"
 
@@ -21,7 +21,7 @@ class UProceduralFoliageComponent;
 typedef TFunction<bool(const UPrimitiveComponent*)> FFoliageTraceFilterFunc;
 
 UCLASS(notplaceable, hidecategories = (Object, Rendering, Mobility), MinimalAPI, NotBlueprintable)
-class AInstancedFoliageActor : public APartitionActor
+class AInstancedFoliageActor : public AISMPartitionActor
 {
 	GENERATED_UCLASS_BODY()
 
