Foliage: Exclude bEditorOnly & bHiddenInGame from Bounds calculation

diff --git a/Engine/Source/Runtime/Foliage/Private/InstancedFoliage.cpp b/Engine/Source/Runtime/Foliage/Private/InstancedFoliage.cpp
--- a/Engine/Source/Runtime/Foliage/Private/InstancedFoliage.cpp
+++ b/Engine/Source/Runtime/Foliage/Private/InstancedFoliage.cpp
@@ -801,7 +801,10 @@ void UFoliageType_Actor::UpdateBounds()
 
 		for (USceneComponent* PreviewComponent : PreviewComponents)
 		{
-			MeshBounds = MeshBounds + PreviewComponent->Bounds;
+			if (!(PreviewComponent->bIsEditorOnly || PreviewComponent->bHiddenInGame))
+			{
+				MeshBounds = MeshBounds + PreviewComponent->Bounds;
+			}
 		}
 	}
 
