Procedural Foliage: Average Normal + Fix non deterministic ZOffset

diff --git a/Engine/Source/Editor/AssetPlacementEdMode/Private/Tools/PlacementBrushToolBase.cpp b/Engine/Source/Editor/AssetPlacementEdMode/Private/Tools/PlacementBrushToolBase.cpp
--- a/Engine/Source/Editor/AssetPlacementEdMode/Private/Tools/PlacementBrushToolBase.cpp
+++ b/Engine/Source/Editor/AssetPlacementEdMode/Private/Tools/PlacementBrushToolBase.cpp
@@ -80,7 +80,7 @@ bool UPlacementBrushToolBase::FindHitResultWithStartAndEndTraceVectors(FHitResul
 
 		return false; };
 
-	return AInstancedFoliageActor::FoliageTrace(EditingWorld, OutHit, FDesiredFoliageInstance(TraceStart, TraceEnd, TraceRadius), NAME_PlacementBrushTool, false, FilterFunc);
+	return AInstancedFoliageActor::FoliageTrace(EditingWorld, OutHit, FDesiredFoliageInstance(TraceStart, TraceEnd, TraceRadius), NAME_PlacementBrushTool, /* bReturnFaceIndex */ false, FilterFunc);
 }
 
 FTransform UPlacementBrushToolBase::GetFinalTransformFromHitLocationAndNormal(const FVector& InLocation, const FVector& InNormal)
diff --git a/Engine/Source/Editor/FoliageEdit/Private/FoliageEdMode.cpp b/Engine/Source/Editor/FoliageEdit/Private/FoliageEdMode.cpp
--- a/Engine/Source/Editor/FoliageEdit/Private/FoliageEdMode.cpp
+++ b/Engine/Source/Editor/FoliageEdit/Private/FoliageEdMode.cpp
@@ -861,7 +861,7 @@ void FEdModeFoliage::FoliageBrushTrace(FEditorViewportClient* ViewportClient, co
 			static FName NAME_FoliageBrush = FName(TEXT("FoliageBrush"));
 			FFoliagePaintingGeometryFilter FilterFunc = FFoliagePaintingGeometryFilter(UISettings);
 
-			if (AInstancedFoliageActor::FoliageTrace(World, Hit, FDesiredFoliageInstance(TraceStart, TraceEnd), NAME_FoliageBrush, false, FilterFunc))
+			if (AInstancedFoliageActor::FoliageTrace(World, Hit, FDesiredFoliageInstance(TraceStart, TraceEnd), NAME_FoliageBrush, /* bReturnFaceIndex */ false, FilterFunc))
 			{
 				UPrimitiveComponent* PrimComp = Hit.Component.Get();
 				if (PrimComp != nullptr && CanPaint(PrimComp->GetComponentLevel()))
@@ -1249,7 +1249,7 @@ void FEdModeFoliage::CalculatePotentialInstances_ThreadSafe(UWorld* InWorld, con
 			TraceFilterFunc = *OverrideGeometryFilter;
 		}
 
-		if (AInstancedFoliageActor::FoliageTrace(InWorld, Hit, DesiredInst, NAME_AddFoliageInstances, true, TraceFilterFunc))
+		if (AInstancedFoliageActor::FoliageTrace(InWorld, Hit, DesiredInst, NAME_AddFoliageInstances, /* bReturnFaceIndex */ true, TraceFilterFunc, /*bAverageNormal*/ true))
 		{
 			float HitWeight = 1.f;
 			const bool bValidInstance = CheckLocationForPotentialInstance_ThreadSafe(Settings, Hit.ImpactPoint, Hit.ImpactNormal)
@@ -1301,7 +1301,7 @@ void FEdModeFoliage::CalculatePotentialInstances(UWorld* InWorld, const UFoliage
 
 		FHitResult Hit;
 		static FName NAME_AddFoliageInstances = FName(TEXT("AddFoliageInstances"));
-		if (AInstancedFoliageActor::FoliageTrace(InWorld, Hit, DesiredInst, NAME_AddFoliageInstances, true, TraceFilterFunc))
+		if (AInstancedFoliageActor::FoliageTrace(InWorld, Hit, DesiredInst, NAME_AddFoliageInstances, /* bReturnFaceIndex */ true, TraceFilterFunc, /*bAverageNormal*/ true))
 		{
 			float HitWeight = 1.f;
 
@@ -2388,7 +2388,7 @@ void FEdModeFoliage::ReapplyInstancesForBrush(UWorld* InWorld, AInstancedFoliage
 				FVector ZAxis = Instance.Rotation.Quaternion().GetAxisZ();
 				FVector Start = Instance.Location + 16.f * ZAxis;
 				FVector End = Instance.Location - 16.f * ZAxis;
-				if (AInstancedFoliageActor::FoliageTrace(InWorld, Hit, FDesiredFoliageInstance(Start, End), NAME_ReapplyInstancesForBrush, true))
+				if (AInstancedFoliageActor::FoliageTrace(InWorld, Hit, FDesiredFoliageInstance(Start, End), NAME_ReapplyInstancesForBrush, /* bReturnFaceIndex */ true, FFoliageTraceFilterFunc(), /*bAverageNormal*/ true))
 				{
 					// Reapply the normal
 					if (bReapplyNormal)
@@ -3094,7 +3094,7 @@ bool FEdModeFoliage::SnapInstanceToGround(AInstancedFoliageActor* InIFA, float A
 
 	FHitResult Hit;
 	static FName NAME_FoliageSnap = FName("FoliageSnap");
-	if (AInstancedFoliageActor::FoliageTrace(InIFA->GetWorld(), Hit, FDesiredFoliageInstance(Start, End), NAME_FoliageSnap))
+	if (AInstancedFoliageActor::FoliageTrace(InIFA->GetWorld(), Hit, FDesiredFoliageInstance(Start, End), NAME_FoliageSnap, /* bReturnFaceIndex */ false, FFoliageTraceFilterFunc(), (Instance.Flags & FOLIAGE_AlignToNormal)))
 	{
 		UPrimitiveComponent* HitComponent = Hit.Component.Get();
 
diff --git a/Engine/Source/Runtime/Foliage/Private/InstancedFoliage.cpp b/Engine/Source/Runtime/Foliage/Private/InstancedFoliage.cpp
--- a/Engine/Source/Runtime/Foliage/Private/InstancedFoliage.cpp
+++ b/Engine/Source/Runtime/Foliage/Private/InstancedFoliage.cpp
@@ -149,6 +149,18 @@ struct FFoliageStaticMesh : public FFoliageImpl
 	void CreateNewComponent(const UFoliageType* InSettings);
 };
 
+struct FFoliagePlacementUtil
+{
+	static int32 GetRandomSeedForPosition(const FVector2D& Position)
+	{
+		// generate a unique random seed for a given position (precision = cm)
+		int32 Xcm = FMath::RoundToInt(Position.X);
+		int32 Ycm = FMath::RoundToInt(Position.Y);
+		// use the int32 hashing function to avoid patterns by spreading out distribution : 
+		return HashCombine(GetTypeHash(Xcm), GetTypeHash(Ycm));
+	}
+};
+
 // Legacy (< FFoliageCustomVersion::CrossLevelBase) serializer
 FArchive& operator<<(FArchive& Ar, FFoliageInstance_Deprecated& Instance)
 {
@@ -406,7 +418,7 @@ bool FFoliageDensityFalloff::IsInstanceFiltered(const FVector2D& InstancePositio
 	check(KeepPointProbability >= 0.f && KeepPointProbability <= 1.f);
 	if (KeepPointProbability < 1.f)
 	{
-		int32 PointSeed = GetRandomSeedForPosition(InstancePosition);
+		int32 PointSeed = FFoliagePlacementUtil::GetRandomSeedForPosition(InstancePosition);
 		FRandomStream LocalRandomStream(PointSeed);
 		float Rand = LocalRandomStream.FRand();
 		return (Rand > KeepPointProbability);
@@ -429,16 +441,6 @@ float FFoliageDensityFalloff::GetDensityFalloffValue(const FVector2D& Position,
 	return KeepPointProbability;
 }
 
-int32 FFoliageDensityFalloff::GetRandomSeedForPosition(const FVector2D& Position)
-{
-	// generate a unique random seed for a given position (precision = cm)
-	int32 Xcm = FMath::RoundToInt(Position.X);
-	int32 Ycm = FMath::RoundToInt(Position.Y);
-	// use the int32 hashing function to avoid patterns by spreading out distribution : 
-	return HashCombine(GetTypeHash(Xcm), GetTypeHash(Ycm));
-}
-
-
 //
 // UFoliageType
 //
@@ -469,6 +471,9 @@ UFoliageType::UFoliageType(const FObjectInitializer& ObjectInitializer)
 	CullDistance.Max = 0;
 	bEnableStaticLighting_DEPRECATED = true;
 	MinimumLayerWeight = 0.5f;
+	AverageNormal = false;
+	AverageNormalSingleComponent = true;
+	AverageNormalSampleCount = 10;
 #if WITH_EDITORONLY_DATA
 	IsSelected = false;
 #endif
@@ -4732,7 +4737,7 @@ void AInstancedFoliageActor::AddReferencedObjects(UObject* InThis, FReferenceCol
 }
 
 #if WITH_EDITOR
-bool AInstancedFoliageActor::FoliageTrace(const UWorld* InWorld, FHitResult& OutHit, const FDesiredFoliageInstance& DesiredInstance, FName InTraceTag, bool InbReturnFaceIndex, const FFoliageTraceFilterFunc& FilterFunc)
+bool AInstancedFoliageActor::FoliageTrace(const UWorld* InWorld, FHitResult& OutHit, const FDesiredFoliageInstance& DesiredInstance, FName InTraceTag, bool InbReturnFaceIndex, const FFoliageTraceFilterFunc& FilterFunc, bool bAverageNormal)
 {
 	SCOPE_CYCLE_COUNTER(STAT_FoliageTrace);
 
@@ -4749,11 +4754,12 @@ bool AInstancedFoliageActor::FoliageTrace(const UWorld* InWorld, FHitResult& Out
 	FCollisionShape SphereShape;
 	SphereShape.SetSphere(DesiredInstance.TraceRadius);
 	InWorld->SweepMultiByObjectType(Hits, StartTrace, DesiredInstance.EndTrace, FQuat::Identity, FCollisionObjectQueryParams(ECC_WorldStatic), SphereShape, QueryParams);
-
-	for (const FHitResult& Hit : Hits)
+	auto ValidateHit = [&DesiredInstance, &FilterFunc](const FHitResult& Hit, FHitResult& OutHit, bool& bOutDiscardHit, bool& bOutInsideProceduralVolumeOrArentUsingOne)
 	{
+		bOutDiscardHit = false;
+		bOutInsideProceduralVolumeOrArentUsingOne = false;
 		const FActorInstanceHandle& HitObjectHandle = Hit.HitObjectHandle;
-				
+
 		// don't place procedural foliage inside an AProceduralFoliageBlockingVolume
 		// this test is first because two of the tests below would otherwise cause the trace to ignore AProceduralFoliageBlockingVolume
 		if (DesiredInstance.PlacementMode == EFoliagePlacementMode::Procedural)
@@ -4783,7 +4789,8 @@ bool AInstancedFoliageActor::FoliageTrace(const UWorld* InWorld, FHitResult& Out
 			}
 			else if (HitObjectHandle.IsValid() && HitObjectHandle.DoesRepresentClass(AProceduralFoliageVolume::StaticClass())) //we never want to collide with our spawning volume
 			{
-				continue;
+				bOutDiscardHit = true;
+				return true;
 			}
 		}
 
@@ -4793,31 +4800,35 @@ bool AInstancedFoliageActor::FoliageTrace(const UWorld* InWorld, FHitResult& Out
 		// In the editor traces can hit "No Collision" type actors, so ugh. (ignore these)
 		if (!HitComponent->IsQueryCollisionEnabled() || HitComponent->GetCollisionResponseToChannel(ECC_WorldStatic) != ECR_Block)
 		{
-			continue;
+			bOutDiscardHit = true;
+			return true;
 		}
 
 		// Don't place foliage on invisible walls / triggers / volumes
 		if (HitComponent->IsA<UBrushComponent>())
 		{
-			continue;
+			bOutDiscardHit = true;
+			return true;
 		}
 
 		// Don't place foliage on itself
 		const AActor* HitActor = Hit.HitObjectHandle.FetchActor();
-		const AInstancedFoliageActor* FoliageActor = Cast<AInstancedFoliageActor>(HitActor);
-		if (!FoliageActor && HitActor && FFoliageHelper::IsOwnedByFoliage(HitActor))
+		const AInstancedFoliageActor* IFA = Cast<AInstancedFoliageActor>(HitActor);
+		if (!IFA && HitActor && FFoliageHelper::IsOwnedByFoliage(HitActor))
 		{
-			FoliageActor = HitActor->GetLevel()->InstancedFoliageActor.Get();
-			if (FoliageActor == nullptr)
+			IFA = HitActor->GetLevel()->InstancedFoliageActor.Get();
+			if (IFA == nullptr)
 			{
-				continue;
+				bOutDiscardHit = true;
+				return true;
 			}
 
-			if (const FFoliageInfo* FoundMeshInfo = FoliageActor->FindInfo(DesiredInstance.FoliageType))
+			if (const FFoliageInfo* FoundMeshInfo = IFA->FindInfo(DesiredInstance.FoliageType))
 			{
 				if (FoundMeshInfo->Implementation->IsOwnedComponent(HitComponent))
 				{
-					continue;
+					bOutDiscardHit = true;
+					return true;
 				}
 			}
 		}
@@ -4825,28 +4836,29 @@ bool AInstancedFoliageActor::FoliageTrace(const UWorld* InWorld, FHitResult& Out
 		if (FilterFunc && FilterFunc(HitComponent) == false)
 		{
 			// supplied filter does not like this component, so keep iterating
-			continue;
+			bOutDiscardHit = true;
+			return true;
 		}
 
-		bool bInsideProceduralVolumeOrArentUsingOne = true;
+		bOutInsideProceduralVolumeOrArentUsingOne = true;
 		if (DesiredInstance.PlacementMode == EFoliagePlacementMode::Procedural && DesiredInstance.ProceduralVolumeBodyInstance)
 		{
 			// We have a procedural volume, so lets make sure we are inside it.
-			bInsideProceduralVolumeOrArentUsingOne = DesiredInstance.ProceduralVolumeBodyInstance->OverlapTest(Hit.ImpactPoint, FQuat::Identity, FCollisionShape::MakeSphere(1.f));	//make sphere of 1cm radius to test if we're in the procedural volume
+			bOutInsideProceduralVolumeOrArentUsingOne = DesiredInstance.ProceduralVolumeBodyInstance->OverlapTest(Hit.ImpactPoint, FQuat::Identity, FCollisionShape::MakeSphere(1.f));	//make sphere of 1cm radius to test if we're in the procedural volume
 		}
 
 		OutHit = Hit;
-			
+		
 		// When placing foliage on other foliage, we need to return the base component of the other foliage, not the foliage component, so that it moves correctly
-		if (FoliageActor)
+		if (IFA)
 		{
-			for (auto& Pair : FoliageActor->FoliageInfos)
+			for (auto& Pair : IFA->FoliageInfos)
 			{
 				const FFoliageInfo& Info = *Pair.Value;
 
 				if (Hit.Item != INDEX_NONE && Info.Implementation->IsOwnedComponent(HitComponent))
 				{
-					OutHit.Component = CastChecked<UPrimitiveComponent>(FoliageActor->InstanceBaseCache.GetInstanceBasePtr(Info.Instances[Hit.Item].BaseId).Get(), ECastCheckedType::NullAllowed);
+					OutHit.Component = CastChecked<UPrimitiveComponent>(IFA->InstanceBaseCache.GetInstanceBasePtr(Info.Instances[Hit.Item].BaseId).Get(), ECastCheckedType::NullAllowed);
 					break;
 				}
 				else
@@ -4854,20 +4866,73 @@ bool AInstancedFoliageActor::FoliageTrace(const UWorld* InWorld, FHitResult& Out
 					int32 InstanceIndex = Info.Implementation->FindIndex(HitComponent);
 					if (InstanceIndex != INDEX_NONE)
 					{
-						OutHit.Component = CastChecked<UPrimitiveComponent>(FoliageActor->InstanceBaseCache.GetInstanceBasePtr(Info.Instances[InstanceIndex].BaseId).Get(), ECastCheckedType::NullAllowed);
+						OutHit.Component = CastChecked<UPrimitiveComponent>(IFA->InstanceBaseCache.GetInstanceBasePtr(Info.Instances[InstanceIndex].BaseId).Get(), ECastCheckedType::NullAllowed);
 						break;
 					}
-				}				
+				}
 			}
 
 			// The foliage we are snapping on doesn't have a valid base
 			if (!OutHit.Component.IsValid())
 			{
-				continue; 
+				bOutDiscardHit = true;
 			}
 		}
 
-		return bInsideProceduralVolumeOrArentUsingOne;
+		return true;
+	};
+
+	for (const FHitResult& Hit : Hits)
+	{
+		bool bOutDiscardHit = false;
+		bool bOutInsideProceduralVolumeOrArentUsingOne = false;
+		if (!ValidateHit(Hit, OutHit, bOutDiscardHit, bOutInsideProceduralVolumeOrArentUsingOne))
+		{
+			return false;
+		}
+					
+		if (bOutDiscardHit)
+		{
+			continue;
+		}
+
+		if (bAverageNormal && DesiredInstance.FoliageType->AverageNormal)
+		{
+			int32 PointSeed = FFoliagePlacementUtil::GetRandomSeedForPosition(FVector2D(Hit.Location));
+			FRandomStream LocalRandomStream(PointSeed);
+			TArray<FHitResult> NormalHits;
+			FVector CumulativeNormal = OutHit.ImpactNormal;
+			FHitResult OutNormalHit;
+			bool bSingleComponent = DesiredInstance.FoliageType->AverageNormalSingleComponent;
+			for (int32 i = 0; i < DesiredInstance.FoliageType->AverageNormalSampleCount; ++i)
+			{
+				const float Angle = LocalRandomStream.FRandRange(0, PI * 2.f);
+				const float SqrtRadius = FMath::Sqrt(LocalRandomStream.FRand()) * DesiredInstance.FoliageType->LowBoundOriginRadius.Z;
+				FVector Offset(SqrtRadius * FMath::Cos(Angle), SqrtRadius* FMath::Sin(Angle), 0.f);
+				NormalHits.Reset();
+				if (InWorld->LineTraceMultiByObjectType(NormalHits, StartTrace + Offset, DesiredInstance.EndTrace + Offset, FCollisionObjectQueryParams(ECC_WorldStatic), QueryParams))
+				{
+					for (const FHitResult& NormalHit : NormalHits)
+					{
+						bool bOutDiscardNormalHit = false;
+						bool bIgnoredParam = false;
+
+						if (ValidateHit(NormalHit, OutNormalHit, bOutDiscardNormalHit, bIgnoredParam))
+						{
+							if (!bOutDiscardNormalHit && (!bSingleComponent || OutNormalHit.Component == OutHit.Component))
+							{
+								CumulativeNormal += OutNormalHit.ImpactNormal;
+								break;
+							}
+						}
+					}
+				}
+			}
+						
+			OutHit.ImpactNormal = CumulativeNormal.GetSafeNormal();
+		}
+						
+		return bOutInsideProceduralVolumeOrArentUsingOne;
 	}
 
 	return false;
@@ -4952,15 +5017,18 @@ bool FPotentialInstance::PlaceInstance(const UWorld* InWorld, const UFoliageType
 	if (DesiredInstance.PlacementMode != EFoliagePlacementMode::Procedural)
 	{
 		Inst.DrawScale3D = Settings->GetRandomScale();
+		Inst.ZOffset = Settings->ZOffset.Interpolate(FMath::FRand());
 	}
 	else
 	{
 		//Procedural foliage uses age to get the scale
 		Inst.DrawScale3D = FVector(Settings->GetScaleForAge(DesiredInstance.Age));
+		
+		// Use a deterministic seed for the offset in Procedural placement so that offset is always the same for the same instance position
+		FRandomStream LocalRandomStream(FFoliagePlacementUtil::GetRandomSeedForPosition(FVector2D(Inst.Location)));
+		Inst.ZOffset = Settings->ZOffset.Interpolate(LocalRandomStream.FRand());
 	}
-
-	Inst.ZOffset = Settings->ZOffset.Interpolate(FMath::FRand());
-
+	
 	Inst.Location = HitLocation;
 
 	if (DesiredInstance.PlacementMode != EFoliagePlacementMode::Procedural)
diff --git a/Engine/Source/Runtime/Foliage/Public/FoliageType.h b/Engine/Source/Runtime/Foliage/Public/FoliageType.h
--- a/Engine/Source/Runtime/Foliage/Public/FoliageType.h
+++ b/Engine/Source/Runtime/Foliage/Public/FoliageType.h
@@ -99,10 +99,6 @@ struct FFoliageDensityFalloff
 
 	FOLIAGE_API bool IsInstanceFiltered(const FVector2D& Position, const FVector2D& Origin, float MaxDistance) const;
 	FOLIAGE_API float GetDensityFalloffValue(const FVector2D& Position, const FVector2D& Origin, float MaxDistance) const;
-
-private:
-
-	static int32 GetRandomSeedForPosition(const FVector2D& Position);
 };
 
 UCLASS(abstract, hidecategories = Object, editinlinenew, MinimalAPI, BlueprintType, Blueprintable)
@@ -215,6 +211,14 @@ public:
 	UPROPERTY(EditAnywhere, BlueprintReadWrite, Category=Placement, meta=(ReapplyCondition="ReapplyAlignToNormal"))
 	uint32 AlignToNormal:1;
 
+	/**	Whether the normal should be averaged on a number of samples around the hit location */
+	UPROPERTY(EditAnywhere, BlueprintReadWrite, Category=Placement, meta=(ToolTip="Will average normal based on Foliage Type base radius (this as a cost as it will do extra line traces)"))
+	uint32 AverageNormal:1;
+
+	/** Average Normal should use all hit components or only the original hit component */
+	UPROPERTY(EditAnywhere, BlueprintReadWrite, Category = Placement, meta = (EditCondition = "AverageNormal", ToolTip = "Whether to discard normals originating from other hit components or not when averaging normals"))
+	uint32 AverageNormalSingleComponent:1;
+
 	/** The maximum angle in degrees that foliage instances will be adjusted away from the vertical */
 	UPROPERTY(EditAnywhere, BlueprintReadWrite, Category=Placement, meta=(UIMin = 0, ClampMin = 0, UIMax = 359, ClampMax = 359, HideBehind="AlignToNormal"))
 	float AlignMaxAngle;
@@ -261,6 +265,10 @@ public:
 	/* The foliage instance's collision bounding box will be scaled by the specified amount before performing the overlap check */
 	UPROPERTY(EditAnywhere, BlueprintReadWrite, AdvancedDisplay, Category=Placement, meta=(HideBehind="CollisionWithWorld"))
 	FVector CollisionScale;
+		
+	/** Line trace count to use around hit location when averaging normal */
+	UPROPERTY(EditAnywhere, BlueprintReadWrite, AdvancedDisplay, Category = Placement, meta=(EditCondition="AverageNormal"))
+	int32 AverageNormalSampleCount;
 
 	UPROPERTY()
 	FBoxSphereBounds MeshBounds;
diff --git a/Engine/Source/Runtime/Foliage/Public/InstancedFoliageActor.h b/Engine/Source/Runtime/Foliage/Public/InstancedFoliageActor.h
--- a/Engine/Source/Runtime/Foliage/Public/InstancedFoliageActor.h
+++ b/Engine/Source/Runtime/Foliage/Public/InstancedFoliageActor.h
@@ -136,7 +136,7 @@ public:
 
 	static FOLIAGE_API AInstancedFoliageActor* GetDefault(UWorld* InWorld);
 
-	static FOLIAGE_API bool FoliageTrace(const UWorld* InWorld, FHitResult& OutHit, const FDesiredFoliageInstance& DesiredInstance, FName InTraceTag = NAME_None, bool InbReturnFaceIndex = false, const FFoliageTraceFilterFunc& FilterFunc = FFoliageTraceFilterFunc());
+	static FOLIAGE_API bool FoliageTrace(const UWorld* InWorld, FHitResult& OutHit, const FDesiredFoliageInstance& DesiredInstance, FName InTraceTag = NAME_None, bool InbReturnFaceIndex = false, const FFoliageTraceFilterFunc& FilterFunc = FFoliageTraceFilterFunc(), bool bAverageNormal = false);
 	static FOLIAGE_API bool CheckCollisionWithWorld(const UWorld* InWorld, const UFoliageType* Settings, const FFoliageInstance& Inst, const FVector& HitNormal, const FVector& HitLocation, UPrimitiveComponent* HitComponent);
 
 	virtual void PreEditUndo() override;
