Fix Selection crash with FoliageISMActor

diff --git a/Engine/Source/Runtime/Foliage/Private/FoliageISMActor.cpp b/Engine/Source/Runtime/Foliage/Private/FoliageISMActor.cpp
--- a/Engine/Source/Runtime/Foliage/Private/FoliageISMActor.cpp
+++ b/Engine/Source/Runtime/Foliage/Private/FoliageISMActor.cpp
@@ -336,7 +336,7 @@ void FFoliageISMActor::SelectInstances(bool bSelect, const TSet<int32>& Selected
 
 int32 FFoliageISMActor::GetInstanceIndexFrom(const UPrimitiveComponent* PrimitiveComponent, int32 ComponentIndex) const
 {
-	if (ComponentIndex != INDEX_NONE)
+	if (IsInitialized() && ComponentIndex != INDEX_NONE)
 	{
 		if (const UInstancedStaticMeshComponent* ISMComponent = Cast<UInstancedStaticMeshComponent>(PrimitiveComponent))
 		{
