Fixed crash when adding a new foliage type.

diff --git a/Engine/Source/Runtime/Foliage/Private/InstancedFoliage.cpp b/Engine/Source/Runtime/Foliage/Private/InstancedFoliage.cpp
--- a/Engine/Source/Runtime/Foliage/Private/InstancedFoliage.cpp
+++ b/Engine/Source/Runtime/Foliage/Private/InstancedFoliage.cpp
@@ -1171,8 +1171,16 @@ void FFoliageStaticMesh::BeginUpdate()
 {
 	if (UpdateDepth == 0)
 	{
-		bPreviousValue = Component->bAutoRebuildTreeOnInstanceChanges;
-		Component->bAutoRebuildTreeOnInstanceChanges = false;
+		if (Component)
+		{
+			bPreviousValue = Component->bAutoRebuildTreeOnInstanceChanges;
+			Component->bAutoRebuildTreeOnInstanceChanges = false;
+		}
+		else
+		{
+			// The default value for HISM component is true, and if we add a component in between the BeginUpdate/EndUpdate pair, it makes sense also.
+			bPreviousValue = true;
+		}
 	}
 	++UpdateDepth;
 }
@@ -1182,7 +1190,7 @@ void FFoliageStaticMesh::EndUpdate()
 	check(UpdateDepth > 0);
 	--UpdateDepth;
 
-	if (UpdateDepth == 0)
+	if (UpdateDepth == 0 && Component)
 	{
 		Component->bAutoRebuildTreeOnInstanceChanges = bPreviousValue;
 
@@ -2241,16 +2249,20 @@ void FFoliageInfo::ReallocateClusters(UFoliageType* InSettings)
 	// Copy the UpdateGuid from the foliage type
 	FoliageTypeUpdateGuid = InSettings->UpdateGuid;
 
-	Implementation->BeginUpdate();
-	// Re-add
+	// Filter instances to re-add
+	TArray<const FFoliageInstance*> InstancesToReAdd;
+	InstancesToReAdd.Reserve(OldInstances.Num());
+
 	for (FFoliageInstance& Instance : OldInstances)
 	{
 		if ((Instance.Flags & FOLIAGE_InstanceDeleted) == 0)
 		{
-			AddInstance(InSettings, Instance);
+			InstancesToReAdd.Add(&Instance);
 		}
 	}
-	Implementation->EndUpdate();
+
+	// Finally, re-add the instances
+	AddInstances(InSettings, InstancesToReAdd);
 	
 	Refresh(true, true);
 }
