Support for editing foliage instances directly in the level editor via typed elements

Note: This is disabled by default, but can be enabled by setting the "TypedElements.EnableFoliageInstanceElements" CVar to "1" (alongside setting "TypedElements.EnableSMInstanceElements" to "1").

Notable changes:
 - AInstancedFoliageActor now implements ISMInstanceManager, and provides an implementation that works for ISMC foliage instances.
   - ISMActor foliage is not yet supported, but also not currently used.
 - ISMInstanceManager now has functions for NotifySMInstanceMovementX and NotifySMInstanceSelectionChanged, which are called from the SMInstance typed element implementation.
 - USMInstanceElementDetailsProxyObject now calls the NotifySMInstanceMovementX functions for transform edits.
 - FFoliageStaticMesh::Reapply now re-uses existing instances where possible, to avoid invalidating mapped FSMInstanceElementId values.
 - FFoliageInfo now keeps track of all moving instances, rather than if any instances are currently moving.
 - Various FFoliageInfo/FFoliageImpl functions now take their indices via TArrayView rather than TArray.
 - UEngineElementsLibrary now has public functions for ReplaceEditorXElementHandles, which exposes the previously private re-instancing logic used by UEngineElementsLibrary::OnObjectsReplaced.
   - This allows static mesh instance elements to be re-instanced if needed (eg, if moving them to an ISMC in a different world partition).

diff --git a/Engine/Source/Editor/FoliageEdit/Private/FoliageEdMode.cpp b/Engine/Source/Editor/FoliageEdit/Private/FoliageEdMode.cpp
--- a/Engine/Source/Editor/FoliageEdit/Private/FoliageEdMode.cpp
+++ b/Engine/Source/Editor/FoliageEdit/Private/FoliageEdMode.cpp
@@ -1942,12 +1942,13 @@ void FEdModeFoliage::PostTransformSelectedInstances(UWorld* InWorld)
 
 void FEdModeFoliage::TransformSelectedInstances(UWorld* InWorld, const FVector& InDrag, const FRotator& InRot, const FVector& InScale, bool bDuplicate)
 {
+	bool bDidMoveFoliage = false;
 	for (TActorIterator<AInstancedFoliageActor> It(InWorld); It; ++It)
 	{
 		AInstancedFoliageActor* IFA = *It;
 		bool bFoundSelection = false;
 
-		IFA->ForEachFoliageInfo([this, IFA, &bFoundSelection, InDrag, InRot, InScale, bDuplicate](UFoliageType* FoliageType, FFoliageInfo& FoliageInfo)
+		IFA->ForEachFoliageInfo([this, IFA, &bDidMoveFoliage, &bFoundSelection, InDrag, InRot, InScale, bDuplicate](UFoliageType* FoliageType, FFoliageInfo& FoliageInfo)
 		{
 			TArray<int32> SelectedIndices = FoliageInfo.SelectedIndices.Array();
 
@@ -1966,8 +1967,11 @@ void FEdModeFoliage::TransformSelectedInstances(UWorld* InWorld, const FVector&
 					OnInstanceCountUpdated(FoliageType);
 				}
 
-				bMoving = true;
-				FoliageInfo.PreMoveInstances(SelectedIndices);
+				if (!bMoving)
+				{
+					FoliageInfo.PreMoveInstances(SelectedIndices);
+				}
+				bDidMoveFoliage = true;
 
 				for (int32 SelectedInstanceIdx : SelectedIndices)
 				{
@@ -1978,8 +1982,8 @@ void FEdModeFoliage::TransformSelectedInstances(UWorld* InWorld, const FVector&
 					Instance.DrawScale3D += InScale;
 				}
 
-				const bool bFinished = false;
-				FoliageInfo.PostMoveInstances(SelectedIndices, bFinished);
+				FoliageInfo.PostMoveInstances(SelectedIndices, /*bFinished*/false);
+				FoliageInfo.PreMoveInstances(SelectedIndices);
 			}
 			return true; // continue iteration
 		});
@@ -1989,6 +1993,11 @@ void FEdModeFoliage::TransformSelectedInstances(UWorld* InWorld, const FVector&
 			IFA->MarkComponentsRenderStateDirty();
 		}
 	}
+
+	if (bDidMoveFoliage)
+	{
+		bMoving = true;
+	}
 }
 
 bool FEdModeFoliage::GetSelectionLocation(UWorld* InWorld, FVector& OutLocation) const
diff --git a/Engine/Source/Editor/LevelEditor/Private/Elements/SMInstance/SMInstanceElementLevelEditorSelectionCustomization.cpp b/Engine/Source/Editor/LevelEditor/Private/Elements/SMInstance/SMInstanceElementLevelEditorSelectionCustomization.cpp
--- a/Engine/Source/Editor/LevelEditor/Private/Elements/SMInstance/SMInstanceElementLevelEditorSelectionCustomization.cpp
+++ b/Engine/Source/Editor/LevelEditor/Private/Elements/SMInstance/SMInstanceElementLevelEditorSelectionCustomization.cpp
@@ -47,9 +47,6 @@ bool FSMInstanceElementLevelEditorSelectionCustomization::SelectElement(const TT
 	}
 
 	UE_LOG(LogSMInstanceLevelEditorSelection, Verbose, TEXT("Selected SMInstance: %s (%s), Index %d"), *SMInstance.GetISMComponent()->GetPathName(), *SMInstance.GetISMComponent()->GetClass()->GetName(), SMInstance.GetISMInstanceIndex());
-
-	// Update the internal selection state for viewport selection rendering
-	SMInstance.GetISMComponent()->SelectInstance(/*bIsSelected*/true, SMInstance.GetISMInstanceIndex());
 	
 	return true;
 }
@@ -69,8 +66,5 @@ bool FSMInstanceElementLevelEditorSelectionCustomization::DeselectElement(const
 
 	UE_LOG(LogSMInstanceLevelEditorSelection, Verbose, TEXT("Deselected SMInstance: %s (%s), Index %d"), *SMInstance.GetISMComponent()->GetPathName(), *SMInstance.GetISMComponent()->GetClass()->GetName(), SMInstance.GetISMInstanceIndex());
 
-	// Update the internal selection state for viewport selection rendering
-	SMInstance.GetISMComponent()->SelectInstance(/*bIsSelected*/false, SMInstance.GetISMInstanceIndex());
-
 	return true;
 }
diff --git a/Engine/Source/Editor/UnrealEd/Private/Elements/SMInstance/SMInstanceElementDetailsProxyObject.cpp b/Engine/Source/Editor/UnrealEd/Private/Elements/SMInstance/SMInstanceElementDetailsProxyObject.cpp
--- a/Engine/Source/Editor/UnrealEd/Private/Elements/SMInstance/SMInstanceElementDetailsProxyObject.cpp
+++ b/Engine/Source/Editor/UnrealEd/Private/Elements/SMInstance/SMInstanceElementDetailsProxyObject.cpp
@@ -43,9 +43,32 @@ void USMInstanceElementDetailsProxyObject::PostEditChangeChainProperty(FProperty
 		{
 			if (FSMInstanceManager SMInstance = GetSMInstance())
 			{
+				if (PropertyChangedEvent.ChangeType == EPropertyChangeType::Interactive)
+				{
+					if (!bIsWithinInteractiveTransformEdit)
+					{
+						SMInstance.NotifySMInstanceMovementStarted();
+					}
+					bIsWithinInteractiveTransformEdit = true;
+				}
+
 				// TODO: Need flag for local/world space, like FComponentTransformDetails
 				SMInstance.SetSMInstanceTransform(Transform, /*bWorldSpace*/false, /*bMarkRenderStateDirty*/true);
 				
+				if (PropertyChangedEvent.ChangeType == EPropertyChangeType::Interactive)
+				{
+					check(bIsWithinInteractiveTransformEdit);
+					SMInstance.NotifySMInstanceMovementOngoing();
+				}
+				else
+				{
+					if (bIsWithinInteractiveTransformEdit)
+					{
+						SMInstance.NotifySMInstanceMovementEnded();
+					}
+					bIsWithinInteractiveTransformEdit = false;
+				}
+
 				GUnrealEd->UpdatePivotLocationForSelection();
 				GUnrealEd->RedrawLevelEditingViewports();
 			}
diff --git a/Engine/Source/Editor/UnrealEd/Private/Elements/SMInstance/SMInstanceElementDetailsProxyObject.h b/Engine/Source/Editor/UnrealEd/Private/Elements/SMInstance/SMInstanceElementDetailsProxyObject.h
--- a/Engine/Source/Editor/UnrealEd/Private/Elements/SMInstance/SMInstanceElementDetailsProxyObject.h
+++ b/Engine/Source/Editor/UnrealEd/Private/Elements/SMInstance/SMInstanceElementDetailsProxyObject.h
@@ -30,4 +30,5 @@ private:
 	uint64 ISMInstanceId = 0;
 
 	FDelegateHandle TickHandle;
+	bool bIsWithinInteractiveTransformEdit = false;
 };
diff --git a/Engine/Source/Runtime/Engine/Classes/Components/InstancedStaticMeshComponent.h b/Engine/Source/Runtime/Engine/Classes/Components/InstancedStaticMeshComponent.h
--- a/Engine/Source/Runtime/Engine/Classes/Components/InstancedStaticMeshComponent.h
+++ b/Engine/Source/Runtime/Engine/Classes/Components/InstancedStaticMeshComponent.h
@@ -458,6 +458,10 @@ protected:
 	virtual bool CanEditSMInstance(const FSMInstanceId& InstanceId) const override;
 	virtual bool GetSMInstanceTransform(const FSMInstanceId& InstanceId, FTransform& OutInstanceTransform, bool bWorldSpace = false) const override;
 	virtual bool SetSMInstanceTransform(const FSMInstanceId& InstanceId, const FTransform& InstanceTransform, bool bWorldSpace = false, bool bMarkRenderStateDirty = false, bool bTeleport = false) override;
+	virtual void NotifySMInstanceMovementStarted(const FSMInstanceId& InstanceId) override;
+	virtual void NotifySMInstanceMovementOngoing(const FSMInstanceId& InstanceId) override;
+	virtual void NotifySMInstanceMovementEnded(const FSMInstanceId& InstanceId) override;
+	virtual void NotifySMInstanceSelectionChanged(const FSMInstanceId& InstanceId, const bool bIsSelected) override;
 	virtual bool DeleteSMInstances(TArrayView<const FSMInstanceId> InstanceIds) override;
 	virtual bool DuplicateSMInstances(TArrayView<const FSMInstanceId> InstanceIds, TArray<FSMInstanceId>& OutNewInstanceIds) override;
 
diff --git a/Engine/Source/Runtime/Engine/Classes/GameFramework/LightWeightInstanceStaticMeshManager.h b/Engine/Source/Runtime/Engine/Classes/GameFramework/LightWeightInstanceStaticMeshManager.h
--- a/Engine/Source/Runtime/Engine/Classes/GameFramework/LightWeightInstanceStaticMeshManager.h
+++ b/Engine/Source/Runtime/Engine/Classes/GameFramework/LightWeightInstanceStaticMeshManager.h
@@ -53,6 +53,10 @@ protected:
 	virtual bool CanEditSMInstance(const FSMInstanceId& InstanceId) const override;
 	virtual bool GetSMInstanceTransform(const FSMInstanceId& InstanceId, FTransform& OutInstanceTransform, bool bWorldSpace = false) const override;
 	virtual bool SetSMInstanceTransform(const FSMInstanceId& InstanceId, const FTransform& InstanceTransform, bool bWorldSpace = false, bool bMarkRenderStateDirty = false, bool bTeleport = false) override;
+	virtual void NotifySMInstanceMovementStarted(const FSMInstanceId& InstanceId) override;
+	virtual void NotifySMInstanceMovementOngoing(const FSMInstanceId& InstanceId) override;
+	virtual void NotifySMInstanceMovementEnded(const FSMInstanceId& InstanceId) override;
+	virtual void NotifySMInstanceSelectionChanged(const FSMInstanceId& InstanceId, const bool bIsSelected) override;
 	virtual bool DeleteSMInstances(TArrayView<const FSMInstanceId> InstanceIds) override;
 	virtual bool DuplicateSMInstances(TArrayView<const FSMInstanceId> InstanceIds, TArray<FSMInstanceId>& OutNewInstanceIds) override;
 
diff --git a/Engine/Source/Runtime/Engine/Private/Elements/Framework/EngineElementsLibrary.cpp b/Engine/Source/Runtime/Engine/Private/Elements/Framework/EngineElementsLibrary.cpp
--- a/Engine/Source/Runtime/Engine/Private/Elements/Framework/EngineElementsLibrary.cpp
+++ b/Engine/Source/Runtime/Engine/Private/Elements/Framework/EngineElementsLibrary.cpp
@@ -75,16 +75,16 @@ FTypedElementHandle AcquireEditorTypedElementHandle(const ObjectClass* Object, T
 }
 
 template <typename ObjectClass>
-TArray<TTuple<const ObjectClass*, const ObjectClass*>> CalculatePotentialObjectReplacements(const TMap<UObject*, UObject*>& ReplacementObjects)
+TMap<const ObjectClass*, const ObjectClass*> CalculatePotentialObjectReplacements(const TMap<UObject*, UObject*>& ReplacementObjects)
 {
-	TArray<TTuple<const ObjectClass*, const ObjectClass*>> PotentialObjectReplacements;
+	TMap<const ObjectClass*, const ObjectClass*> PotentialObjectReplacements;
 
 	for (const TTuple<UObject*, UObject*>& ReplacementObjectPair : ReplacementObjects)
 	{
 		if (const ObjectClass* OldObject = Cast<ObjectClass>(ReplacementObjectPair.Key))
 		{
 			const ObjectClass* NewObject = Cast<ObjectClass>(ReplacementObjectPair.Value);
-			PotentialObjectReplacements.Add(MakeTuple(OldObject, NewObject));
+			PotentialObjectReplacements.Add(OldObject, NewObject);
 		}
 	}
 
@@ -92,7 +92,7 @@ TArray<TTuple<const ObjectClass*, const ObjectClass*>> CalculatePotentialObjectR
 }
 
 template <typename ElementDataType, typename KeyDataType>
-void ReplaceEditorTypedElementHandles(TArray<FTypedElementHandle>& OutUpdatedElements, const TArray<TTuple<KeyDataType, KeyDataType>>& ReplacementKeys, TTypedElementOwnerStore<ElementDataType, KeyDataType>& ElementOwnerStore, TFunctionRef<void(KeyDataType, TTypedElementOwner<ElementDataType>&)> UpdateElement, TFunctionRef<void(KeyDataType, TTypedElementOwner<ElementDataType>&)> DestroyElement)
+void ReplaceEditorTypedElementHandles(TArray<FTypedElementHandle>& OutUpdatedElements, const TMap<KeyDataType, KeyDataType>& ReplacementKeys, TTypedElementOwnerStore<ElementDataType, KeyDataType>& ElementOwnerStore, TFunctionRef<void(KeyDataType, TTypedElementOwner<ElementDataType>&)> UpdateElement, TFunctionRef<void(KeyDataType, TTypedElementOwner<ElementDataType>&)> DestroyElement)
 {
 	for (const TTuple<KeyDataType, KeyDataType>& ReplacementKeyPair : ReplacementKeys)
 	{
@@ -176,31 +176,22 @@ void UEngineElementsLibrary::OnObjectsReplaced(const TMap<UObject*, UObject*>& I
 	TArray<FTypedElementHandle> UpdatedElements;
 
 	{
-		const TArray<TTuple<const UObject*, const UObject*>> PotentialObjectReplacements = EngineElementsLibraryUtil::CalculatePotentialObjectReplacements<UObject>(InReplacementObjects);
-		EngineElementsLibraryUtil::ReplaceEditorTypedElementHandles<FObjectElementData, const UObject*>(UpdatedElements, PotentialObjectReplacements, GObjectElementOwnerStore, [](const UObject* InObject, TTypedElementOwner<FObjectElementData>& InOutObjectElement)
-		{
-			InOutObjectElement.GetDataChecked().Object = const_cast<UObject*>(InObject);
-		}, &UEngineElementsLibrary::DestroyObjectElement);
+		const TMap<const UObject*, const UObject*> PotentialObjectReplacements = EngineElementsLibraryUtil::CalculatePotentialObjectReplacements<UObject>(InReplacementObjects);
+		ReplaceEditorObjectElementHandlesImpl(PotentialObjectReplacements, UpdatedElements);
 	}
 
 	{
-		const TArray<TTuple<const AActor*, const AActor*>> PotentialActorReplacements = EngineElementsLibraryUtil::CalculatePotentialObjectReplacements<AActor>(InReplacementObjects);
-		EngineElementsLibraryUtil::ReplaceEditorTypedElementHandles<FActorElementData, const AActor*>(UpdatedElements, PotentialActorReplacements, GActorElementOwnerStore, [](const AActor* InActor, TTypedElementOwner<FActorElementData>& InOutActorElement)
-		{
-			InOutActorElement.GetDataChecked().Actor = const_cast<AActor*>(InActor);
-		}, &UEngineElementsLibrary::DestroyActorElement);
+		const TMap<const AActor*, const AActor*> PotentialActorReplacements = EngineElementsLibraryUtil::CalculatePotentialObjectReplacements<AActor>(InReplacementObjects);
+		ReplaceEditorActorElementHandlesImpl(PotentialActorReplacements, UpdatedElements);
 	}
 
 	{
-		const TArray<TTuple<const UActorComponent*, const UActorComponent*>> PotentialComponentReplacements = EngineElementsLibraryUtil::CalculatePotentialObjectReplacements<UActorComponent>(InReplacementObjects);
-		EngineElementsLibraryUtil::ReplaceEditorTypedElementHandles<FComponentElementData, const UActorComponent*>(UpdatedElements, PotentialComponentReplacements, GComponentElementOwnerStore, [](const UActorComponent* InComponent, TTypedElementOwner<FComponentElementData>& InOutComponentElement)
-		{
-			InOutComponentElement.GetDataChecked().Component = const_cast<UActorComponent*>(InComponent);
-		}, &UEngineElementsLibrary::DestroyComponentElement);
+		const TMap<const UActorComponent*, const UActorComponent*> PotentialComponentReplacements = EngineElementsLibraryUtil::CalculatePotentialObjectReplacements<UActorComponent>(InReplacementObjects);
+		ReplaceEditorComponentElementHandlesImpl(PotentialComponentReplacements, UpdatedElements);
 	}
 
 	{
-		TArray<TTuple<FSMInstanceElementId, FSMInstanceElementId>> PotentialSMInstanceReplacements;
+		TMap<FSMInstanceElementId, FSMInstanceElementId> PotentialSMInstanceReplacements;
 
 		FSMInstanceElementIdMap& SMInstanceElementIdMap = FSMInstanceElementIdMap::Get();
 		for (const TTuple<UObject*, UObject*>& ReplacementObjectPair : InReplacementObjects)
@@ -231,10 +222,7 @@ void UEngineElementsLibrary::OnObjectsReplaced(const TMap<UObject*, UObject*>& I
 			}
 		}
 
-		EngineElementsLibraryUtil::ReplaceEditorTypedElementHandles<FSMInstanceElementData, FSMInstanceElementId>(UpdatedElements, PotentialSMInstanceReplacements, GSMInstanceElementOwnerStore, [](FSMInstanceElementId InSMInstanceElementId, TTypedElementOwner<FSMInstanceElementData>& InOutSMInstanceElement)
-		{
-			InOutSMInstanceElement.GetDataChecked().InstanceElementId = InSMInstanceElementId;
-		}, &UEngineElementsLibrary::DestroySMInstanceElement);
+		ReplaceEditorSMInstanceElementHandlesImpl(PotentialSMInstanceReplacements, UpdatedElements);
 	}
 
 	if (UpdatedElements.Num() > 0)
@@ -242,6 +230,38 @@ void UEngineElementsLibrary::OnObjectsReplaced(const TMap<UObject*, UObject*>& I
 		UTypedElementRegistry::GetInstance()->OnElementUpdated().Broadcast(UpdatedElements);
 	}
 }
+
+void UEngineElementsLibrary::ReplaceEditorObjectElementHandlesImpl(const TMap<const UObject*, const UObject*>& ReplacementObjects, TArray<FTypedElementHandle>& OutUpdatedElements)
+{
+	EngineElementsLibraryUtil::ReplaceEditorTypedElementHandles<FObjectElementData, const UObject*>(OutUpdatedElements, ReplacementObjects, GObjectElementOwnerStore, [](const UObject* InObject, TTypedElementOwner<FObjectElementData>& InOutObjectElement)
+	{
+		InOutObjectElement.GetDataChecked().Object = const_cast<UObject*>(InObject);
+	}, &UEngineElementsLibrary::DestroyObjectElement);
+}
+
+void UEngineElementsLibrary::ReplaceEditorActorElementHandlesImpl(const TMap<const AActor*, const AActor*>& ReplacementActors, TArray<FTypedElementHandle>& OutUpdatedElements)
+{
+	EngineElementsLibraryUtil::ReplaceEditorTypedElementHandles<FActorElementData, const AActor*>(OutUpdatedElements, ReplacementActors, GActorElementOwnerStore, [](const AActor* InActor, TTypedElementOwner<FActorElementData>& InOutActorElement)
+	{
+		InOutActorElement.GetDataChecked().Actor = const_cast<AActor*>(InActor);
+	}, &UEngineElementsLibrary::DestroyActorElement);
+}
+
+void UEngineElementsLibrary::ReplaceEditorComponentElementHandlesImpl(const TMap<const UActorComponent*, const UActorComponent*>& ReplacementComponents, TArray<FTypedElementHandle>& OutUpdatedElements)
+{
+	EngineElementsLibraryUtil::ReplaceEditorTypedElementHandles<FComponentElementData, const UActorComponent*>(OutUpdatedElements, ReplacementComponents, GComponentElementOwnerStore, [](const UActorComponent* InComponent, TTypedElementOwner<FComponentElementData>& InOutComponentElement)
+	{
+		InOutComponentElement.GetDataChecked().Component = const_cast<UActorComponent*>(InComponent);
+	}, &UEngineElementsLibrary::DestroyComponentElement);
+}
+
+void UEngineElementsLibrary::ReplaceEditorSMInstanceElementHandlesImpl(const TMap<FSMInstanceElementId, FSMInstanceElementId>& ReplacementSMInstanceIds, TArray<FTypedElementHandle>& OutUpdatedElements)
+{
+	EngineElementsLibraryUtil::ReplaceEditorTypedElementHandles<FSMInstanceElementData, FSMInstanceElementId>(OutUpdatedElements, ReplacementSMInstanceIds, GSMInstanceElementOwnerStore, [](FSMInstanceElementId InSMInstanceElementId, TTypedElementOwner<FSMInstanceElementData>& InOutSMInstanceElement)
+	{
+		InOutSMInstanceElement.GetDataChecked().InstanceElementId = InSMInstanceElementId;
+	}, &UEngineElementsLibrary::DestroySMInstanceElement);
+}
 #endif
 
 TTypedElementOwner<FObjectElementData> UEngineElementsLibrary::CreateObjectElement(const UObject* InObject)
@@ -300,6 +320,16 @@ FTypedElementHandle UEngineElementsLibrary::AcquireEditorObjectElementHandle(con
 {
 	return EngineElementsLibraryUtil::AcquireEditorTypedElementHandle<UObject, FObjectElementData>(Object, GObjectElementOwnerStore, &UEngineElementsLibrary::CreateObjectElement, bAllowCreate);
 }
+
+void UEngineElementsLibrary::ReplaceEditorObjectElementHandles(const TMap<const UObject*, const UObject*>& ReplacementObjects)
+{
+	TArray<FTypedElementHandle> UpdatedElements;
+	ReplaceEditorObjectElementHandlesImpl(ReplacementObjects, UpdatedElements);
+	if (UpdatedElements.Num() > 0)
+	{
+		UTypedElementRegistry::GetInstance()->OnElementUpdated().Broadcast(UpdatedElements);
+	}
+}
 #endif
 
 TTypedElementOwner<FActorElementData> UEngineElementsLibrary::CreateActorElement(const AActor* InActor)
@@ -336,6 +366,16 @@ FTypedElementHandle UEngineElementsLibrary::AcquireEditorActorElementHandle(cons
 {
 	return EngineElementsLibraryUtil::AcquireEditorTypedElementHandle<AActor, FActorElementData>(Actor, GActorElementOwnerStore, &UEngineElementsLibrary::CreateActorElement, bAllowCreate);
 }
+
+void UEngineElementsLibrary::ReplaceEditorActorElementHandles(const TMap<const AActor*, const AActor*>& ReplacementActors)
+{
+	TArray<FTypedElementHandle> UpdatedElements;
+	ReplaceEditorActorElementHandlesImpl(ReplacementActors, UpdatedElements);
+	if (UpdatedElements.Num() > 0)
+	{
+		UTypedElementRegistry::GetInstance()->OnElementUpdated().Broadcast(UpdatedElements);
+	}
+}
 #endif
 
 TTypedElementOwner<FComponentElementData> UEngineElementsLibrary::CreateComponentElement(const UActorComponent* InComponent)
@@ -372,6 +412,16 @@ FTypedElementHandle UEngineElementsLibrary::AcquireEditorComponentElementHandle(
 {
 	return EngineElementsLibraryUtil::AcquireEditorTypedElementHandle<UActorComponent, FComponentElementData>(Component, GComponentElementOwnerStore, &UEngineElementsLibrary::CreateComponentElement, bAllowCreate);
 }
+
+void UEngineElementsLibrary::ReplaceEditorComponentElementHandles(const TMap<const UActorComponent*, const UActorComponent*>& ReplacementComponents)
+{
+	TArray<FTypedElementHandle> UpdatedElements;
+	ReplaceEditorComponentElementHandlesImpl(ReplacementComponents, UpdatedElements);
+	if (UpdatedElements.Num() > 0)
+	{
+		UTypedElementRegistry::GetInstance()->OnElementUpdated().Broadcast(UpdatedElements);
+	}
+}
 #endif
 
 TTypedElementOwner<FSMInstanceElementData> UEngineElementsLibrary::CreateSMInstanceElement(const FSMInstanceId& InSMInstanceId)
@@ -452,4 +502,28 @@ FTypedElementHandle UEngineElementsLibrary::AcquireEditorSMInstanceElementHandle
 
 	return FTypedElementHandle();
 }
+
+void UEngineElementsLibrary::ReplaceEditorSMInstanceElementHandles(const TMap<FSMInstanceId, FSMInstanceId>& ReplacementSMInstanceIds)
+{
+	TMap<FSMInstanceElementId, FSMInstanceElementId> PotentialSMInstanceReplacements;
+	PotentialSMInstanceReplacements.Reserve(ReplacementSMInstanceIds.Num());
+
+	FSMInstanceElementIdMap& SMInstanceElementIdMap = FSMInstanceElementIdMap::Get();
+	for (const TTuple<FSMInstanceId, FSMInstanceId>& ReplacementSMInstanceIdPair : ReplacementSMInstanceIds)
+	{
+		// Find, but don't create, the element ID for the source instance, as no element can exist if the element ID has no mapping
+		if (FSMInstanceElementId OldSMInstanceElementId = SMInstanceElementIdMap.GetSMInstanceElementIdFromSMInstanceId(ReplacementSMInstanceIdPair.Key, /*bAllowCreate*/false))
+		{
+			FSMInstanceElementId NewSMInstanceElementId = SMInstanceElementIdMap.GetSMInstanceElementIdFromSMInstanceId(ReplacementSMInstanceIdPair.Value, /*bAllowCreate*/true);
+			PotentialSMInstanceReplacements.Add(MoveTemp(OldSMInstanceElementId), MoveTemp(NewSMInstanceElementId));
+		}
+	}
+
+	TArray<FTypedElementHandle> UpdatedElements;
+	ReplaceEditorSMInstanceElementHandlesImpl(PotentialSMInstanceReplacements, UpdatedElements);
+	if (UpdatedElements.Num() > 0)
+	{
+		UTypedElementRegistry::GetInstance()->OnElementUpdated().Broadcast(UpdatedElements);
+	}
+}
 #endif
diff --git a/Engine/Source/Runtime/Engine/Private/Elements/SMInstance/SMInstanceElementSelectionInterface.cpp b/Engine/Source/Runtime/Engine/Private/Elements/SMInstance/SMInstanceElementSelectionInterface.cpp
--- a/Engine/Source/Runtime/Engine/Private/Elements/SMInstance/SMInstanceElementSelectionInterface.cpp
+++ b/Engine/Source/Runtime/Engine/Private/Elements/SMInstance/SMInstanceElementSelectionInterface.cpp
@@ -4,3 +4,25 @@
 #include "Elements/SMInstance/SMInstanceElementData.h"
 #include "Components/InstancedStaticMeshComponent.h"
 #include "Elements/Framework/TypedElementList.h"
+
+bool USMInstanceElementSelectionInterface::SelectElement(const FTypedElementHandle& InElementHandle, UTypedElementList* InSelectionSet, const FTypedElementSelectionOptions& InSelectionOptions)
+{
+	FSMInstanceManager SMInstance = SMInstanceElementDataUtil::GetSMInstanceFromHandle(InElementHandle);
+	if (SMInstance && Super::SelectElement(InElementHandle, InSelectionSet, InSelectionOptions))
+	{
+		SMInstance.NotifySMInstanceSelectionChanged(/*bIsSelected*/true);
+		return true;
+	}
+	return false;
+}
+
+bool USMInstanceElementSelectionInterface::DeselectElement(const FTypedElementHandle& InElementHandle, UTypedElementList* InSelectionSet, const FTypedElementSelectionOptions& InSelectionOptions)
+{
+	FSMInstanceManager SMInstance = SMInstanceElementDataUtil::GetSMInstanceFromHandle(InElementHandle);
+	if (SMInstance && Super::DeselectElement(InElementHandle, InSelectionSet, InSelectionOptions))
+	{
+		SMInstance.NotifySMInstanceSelectionChanged(/*bIsSelected*/false);
+		return true;
+	}
+	return false;
+}
diff --git a/Engine/Source/Runtime/Engine/Private/Elements/SMInstance/SMInstanceElementWorldInterface.cpp b/Engine/Source/Runtime/Engine/Private/Elements/SMInstance/SMInstanceElementWorldInterface.cpp
--- a/Engine/Source/Runtime/Engine/Private/Elements/SMInstance/SMInstanceElementWorldInterface.cpp
+++ b/Engine/Source/Runtime/Engine/Private/Elements/SMInstance/SMInstanceElementWorldInterface.cpp
@@ -86,3 +86,27 @@ bool USMInstanceElementWorldInterface::SetRelativeTransform(const FTypedElementH
 	const FSMInstanceManager SMInstance = SMInstanceElementDataUtil::GetSMInstanceFromHandle(InElementHandle);
 	return SMInstance && SMInstance.SetSMInstanceTransform(InTransform, /*bWorldSpace*/false, /*bMarkRenderStateDirty*/true);
 }
+
+void USMInstanceElementWorldInterface::NotifyMovementStarted(const FTypedElementHandle& InElementHandle)
+{
+	if (FSMInstanceManager SMInstance = SMInstanceElementDataUtil::GetSMInstanceFromHandle(InElementHandle))
+	{
+		SMInstance.NotifySMInstanceMovementStarted();
+	}
+}
+
+void USMInstanceElementWorldInterface::NotifyMovementOngoing(const FTypedElementHandle& InElementHandle)
+{
+	if (FSMInstanceManager SMInstance = SMInstanceElementDataUtil::GetSMInstanceFromHandle(InElementHandle))
+	{
+		SMInstance.NotifySMInstanceMovementOngoing();
+	}
+}
+
+void USMInstanceElementWorldInterface::NotifyMovementEnded(const FTypedElementHandle& InElementHandle)
+{
+	if (FSMInstanceManager SMInstance = SMInstanceElementDataUtil::GetSMInstanceFromHandle(InElementHandle))
+	{
+		SMInstance.NotifySMInstanceMovementEnded();
+	}
+}
diff --git a/Engine/Source/Runtime/Engine/Private/ISMPartition/ISMPartitionActor.cpp b/Engine/Source/Runtime/Engine/Private/ISMPartition/ISMPartitionActor.cpp
--- a/Engine/Source/Runtime/Engine/Private/ISMPartition/ISMPartitionActor.cpp
+++ b/Engine/Source/Runtime/Engine/Private/ISMPartition/ISMPartitionActor.cpp
@@ -651,7 +651,7 @@ void AISMPartitionActor::GetClientComponents(const FISMClientHandle& Handle, TAr
 
 #endif
 
-TScriptInterface<ISMInstanceManager> AISMPartitionActor::GetSMInstanceManager(const FSMInstanceId& InstanceId) const
+TScriptInterface<ISMInstanceManager> AISMPartitionActor::GetSMInstanceManager(const FSMInstanceId& InstanceId)
 {
 	// Disable per-instance editing for now...
 	return nullptr;
diff --git a/Engine/Source/Runtime/Engine/Private/InstancedStaticMesh.cpp b/Engine/Source/Runtime/Engine/Private/InstancedStaticMesh.cpp
--- a/Engine/Source/Runtime/Engine/Private/InstancedStaticMesh.cpp
+++ b/Engine/Source/Runtime/Engine/Private/InstancedStaticMesh.cpp
@@ -3936,6 +3936,27 @@ bool UInstancedStaticMeshComponent::SetSMInstanceTransform(const FSMInstanceId&
 	return UpdateInstanceTransform(InstanceId.InstanceIndex, InstanceTransform, bWorldSpace, bMarkRenderStateDirty, bTeleport);
 }
 
+void UInstancedStaticMeshComponent::NotifySMInstanceMovementStarted(const FSMInstanceId& InstanceId)
+{
+	check(InstanceId.ISMComponent == this);
+}
+
+void UInstancedStaticMeshComponent::NotifySMInstanceMovementOngoing(const FSMInstanceId& InstanceId)
+{
+	check(InstanceId.ISMComponent == this);
+}
+
+void UInstancedStaticMeshComponent::NotifySMInstanceMovementEnded(const FSMInstanceId& InstanceId)
+{
+	check(InstanceId.ISMComponent == this);
+}
+
+void UInstancedStaticMeshComponent::NotifySMInstanceSelectionChanged(const FSMInstanceId& InstanceId, const bool bIsSelected)
+{
+	check(InstanceId.ISMComponent == this);
+	SelectInstance(bIsSelected, InstanceId.InstanceIndex);
+}
+
 bool UInstancedStaticMeshComponent::DeleteSMInstances(TArrayView<const FSMInstanceId> InstanceIds)
 {
 	TArray<int32> InstanceIndices;
diff --git a/Engine/Source/Runtime/Engine/Private/LightWeightInstanceStaticMeshManager.cpp b/Engine/Source/Runtime/Engine/Private/LightWeightInstanceStaticMeshManager.cpp
--- a/Engine/Source/Runtime/Engine/Private/LightWeightInstanceStaticMeshManager.cpp
+++ b/Engine/Source/Runtime/Engine/Private/LightWeightInstanceStaticMeshManager.cpp
@@ -232,6 +232,38 @@ bool ALightWeightInstanceStaticMeshManager::SetSMInstanceTransform(const FSMInst
 	return false;
 }
 
+void ALightWeightInstanceStaticMeshManager::NotifySMInstanceMovementStarted(const FSMInstanceId& InstanceId)
+{
+	check(InstanceId.ISMComponent == InstancedStaticMeshComponent);
+
+	ISMInstanceManager* InstanceManager = InstancedStaticMeshComponent;
+	InstanceManager->NotifySMInstanceMovementStarted(InstanceId);
+}
+
+void ALightWeightInstanceStaticMeshManager::NotifySMInstanceMovementOngoing(const FSMInstanceId& InstanceId)
+{
+	check(InstanceId.ISMComponent == InstancedStaticMeshComponent);
+
+	ISMInstanceManager* InstanceManager = InstancedStaticMeshComponent;
+	InstanceManager->NotifySMInstanceMovementOngoing(InstanceId);
+}
+
+void ALightWeightInstanceStaticMeshManager::NotifySMInstanceMovementEnded(const FSMInstanceId& InstanceId)
+{
+	check(InstanceId.ISMComponent == InstancedStaticMeshComponent);
+
+	ISMInstanceManager* InstanceManager = InstancedStaticMeshComponent;
+	InstanceManager->NotifySMInstanceMovementEnded(InstanceId);
+}
+
+void ALightWeightInstanceStaticMeshManager::NotifySMInstanceSelectionChanged(const FSMInstanceId& InstanceId, const bool bIsSelected)
+{
+	check(InstanceId.ISMComponent == InstancedStaticMeshComponent);
+
+	ISMInstanceManager* InstanceManager = InstancedStaticMeshComponent;
+	InstanceManager->NotifySMInstanceSelectionChanged(InstanceId, bIsSelected);
+}
+
 bool ALightWeightInstanceStaticMeshManager::DeleteSMInstances(TArrayView<const FSMInstanceId> InstanceIds)
 {
 	TArray<int32> DataIndices;
diff --git a/Engine/Source/Runtime/Engine/Public/Elements/Framework/EngineElementsLibrary.h b/Engine/Source/Runtime/Engine/Public/Elements/Framework/EngineElementsLibrary.h
--- a/Engine/Source/Runtime/Engine/Public/Elements/Framework/EngineElementsLibrary.h
+++ b/Engine/Source/Runtime/Engine/Public/Elements/Framework/EngineElementsLibrary.h
@@ -39,6 +39,7 @@ public:
 	static void DestroyEditorObjectElement(const UObject* Object);
 	UFUNCTION(BlueprintCallable, BlueprintPure=false, Category="TypedElementFramework|Object", meta=(ScriptMethod="AcquireEditorElementHandle"))
 	static FTypedElementHandle AcquireEditorObjectElementHandle(const UObject* Object, const bool bAllowCreate = true);
+	static void ReplaceEditorObjectElementHandles(const TMap<const UObject*, const UObject*>& ReplacementObjects);
 #endif
 
 	static TTypedElementOwner<FActorElementData> CreateActorElement(const AActor* InActor);
@@ -48,6 +49,7 @@ public:
 	static void DestroyEditorActorElement(const AActor* Actor);
 	UFUNCTION(BlueprintCallable, Category="TypedElementFramework|Actor", meta=(ScriptMethod="AcquireEditorElementHandle"))
 	static FTypedElementHandle AcquireEditorActorElementHandle(const AActor* Actor, const bool bAllowCreate = true);
+	static void ReplaceEditorActorElementHandles(const TMap<const AActor*, const AActor*>& ReplacementActors);
 #endif
 
 	static TTypedElementOwner<FComponentElementData> CreateComponentElement(const UActorComponent* InComponent);
@@ -57,6 +59,7 @@ public:
 	static void DestroyEditorComponentElement(const UActorComponent* Component);
 	UFUNCTION(BlueprintCallable, BlueprintPure=false, Category="TypedElementFramework|Component", meta=(ScriptMethod="AcquireEditorElementHandle"))
 	static FTypedElementHandle AcquireEditorComponentElementHandle(const UActorComponent* Component, const bool bAllowCreate = true);
+	static void ReplaceEditorComponentElementHandles(const TMap<const UActorComponent*, const UActorComponent*>& ReplacementComponents);
 #endif
 
 	static TTypedElementOwner<FSMInstanceElementData> CreateSMInstanceElement(const FSMInstanceId& InSMInstanceId);
@@ -67,6 +70,7 @@ public:
 	UFUNCTION(BlueprintCallable, BlueprintPure=false, Category="TypedElementFramework|StaticMeshInstance", meta=(ScriptMethod="AcquireEditorElementHandle"))
 	static FTypedElementHandle AcquireEditorSMInstanceElementHandle(const UInstancedStaticMeshComponent* ISMComponent, const int32 InstanceIndex, const bool bAllowCreate = true);
 	static FTypedElementHandle AcquireEditorSMInstanceElementHandle(const FSMInstanceId& SMInstanceId, const bool bAllowCreate = true);
+	static void ReplaceEditorSMInstanceElementHandles(const TMap<FSMInstanceId, FSMInstanceId>& ReplacementSMInstanceIds);
 #endif
 
 private:
@@ -74,5 +78,9 @@ private:
 #if WITH_EDITOR
 	static void DestroyUnreachableEditorObjectElements();
 	static void OnObjectsReplaced(const TMap<UObject*, UObject*>& InReplacementObjects);
+	static void ReplaceEditorObjectElementHandlesImpl(const TMap<const UObject*, const UObject*>& ReplacementObjects, TArray<FTypedElementHandle>& OutUpdatedElements);
+	static void ReplaceEditorActorElementHandlesImpl(const TMap<const AActor*, const AActor*>& ReplacementActors, TArray<FTypedElementHandle>& OutUpdatedElements);
+	static void ReplaceEditorComponentElementHandlesImpl(const TMap<const UActorComponent*, const UActorComponent*>& ReplacementComponents, TArray<FTypedElementHandle>& OutUpdatedElements);
+	static void ReplaceEditorSMInstanceElementHandlesImpl(const TMap<FSMInstanceElementId, FSMInstanceElementId>& ReplacementSMInstanceIds, TArray<FTypedElementHandle>& OutUpdatedElements);
 #endif
 };
diff --git a/Engine/Source/Runtime/Engine/Public/Elements/SMInstance/SMInstanceElementSelectionInterface.h b/Engine/Source/Runtime/Engine/Public/Elements/SMInstance/SMInstanceElementSelectionInterface.h
--- a/Engine/Source/Runtime/Engine/Public/Elements/SMInstance/SMInstanceElementSelectionInterface.h
+++ b/Engine/Source/Runtime/Engine/Public/Elements/SMInstance/SMInstanceElementSelectionInterface.h
@@ -9,4 +9,8 @@ UCLASS()
 class ENGINE_API USMInstanceElementSelectionInterface : public UTypedElementSelectionInterface
 {
 	GENERATED_BODY()
+
+public:
+	virtual bool SelectElement(const FTypedElementHandle& InElementHandle, UTypedElementList* InSelectionSet, const FTypedElementSelectionOptions& InSelectionOptions) override;
+	virtual bool DeselectElement(const FTypedElementHandle& InElementHandle, UTypedElementList* InSelectionSet, const FTypedElementSelectionOptions& InSelectionOptions) override;
 };
diff --git a/Engine/Source/Runtime/Engine/Public/Elements/SMInstance/SMInstanceElementWorldInterface.h b/Engine/Source/Runtime/Engine/Public/Elements/SMInstance/SMInstanceElementWorldInterface.h
--- a/Engine/Source/Runtime/Engine/Public/Elements/SMInstance/SMInstanceElementWorldInterface.h
+++ b/Engine/Source/Runtime/Engine/Public/Elements/SMInstance/SMInstanceElementWorldInterface.h
@@ -20,4 +20,7 @@ public:
 	virtual bool SetWorldTransform(const FTypedElementHandle& InElementHandle, const FTransform& InTransform) override;
 	virtual bool GetRelativeTransform(const FTypedElementHandle& InElementHandle, FTransform& OutTransform) override;
 	virtual bool SetRelativeTransform(const FTypedElementHandle& InElementHandle, const FTransform& InTransform) override;
+	virtual void NotifyMovementStarted(const FTypedElementHandle& InElementHandle) override;
+	virtual void NotifyMovementOngoing(const FTypedElementHandle& InElementHandle) override;
+	virtual void NotifyMovementEnded(const FTypedElementHandle& InElementHandle) override;
 };
diff --git a/Engine/Source/Runtime/Engine/Public/Elements/SMInstance/SMInstanceManager.h b/Engine/Source/Runtime/Engine/Public/Elements/SMInstance/SMInstanceManager.h
--- a/Engine/Source/Runtime/Engine/Public/Elements/SMInstance/SMInstanceManager.h
+++ b/Engine/Source/Runtime/Engine/Public/Elements/SMInstance/SMInstanceManager.h
@@ -43,6 +43,30 @@ public:
 	 */
 	virtual bool SetSMInstanceTransform(const FSMInstanceId& InstanceId, const FTransform& InstanceTransform, bool bWorldSpace = false, bool bMarkRenderStateDirty = false, bool bTeleport = false) = 0;
 
+	/**
+	 * Notify that the given static mesh instance is about to be moved.
+	 * @note This gives the manager a chance to start a move operation, to avoid performing repeated work until the move is finished.
+	 */
+	virtual void NotifySMInstanceMovementStarted(const FSMInstanceId& InstanceId) = 0;
+
+	/**
+	 * Notify that the given static mesh instance is currently being moved.
+	 * @note This gives the manager a chance to update a move operation.
+	 */
+	virtual void NotifySMInstanceMovementOngoing(const FSMInstanceId& InstanceId) = 0;
+
+	/**
+	 * Notify that the given static mesh instance is done being moved.
+	 * @note This gives the manager a chance to end a move operation.
+	 */
+	virtual void NotifySMInstanceMovementEnded(const FSMInstanceId& InstanceId) = 0;
+
+	/**
+	 * Notify that the given static mesh instance selection state has changed.
+	 * @note This gives the manager a chance to sync any internal selection state.
+	 */
+	virtual void NotifySMInstanceSelectionChanged(const FSMInstanceId& InstanceId, const bool bIsSelected) = 0;
+
 	/**
 	 * Can the given static mesh instance be deleted?
 	 * @return True if it can be deleted, false otherwise.
@@ -93,7 +117,7 @@ public:
 	 * Attempt to get the instance manager associated with the given static mesh instance, if any.
 	 * @return The instance manager, or null if there is no instance manager associated with the given instance.
 	 */
-	virtual TScriptInterface<ISMInstanceManager> GetSMInstanceManager(const FSMInstanceId& InstanceId) const = 0;
+	virtual TScriptInterface<ISMInstanceManager> GetSMInstanceManager(const FSMInstanceId& InstanceId) = 0;
 };
 
 /**
@@ -142,6 +166,10 @@ public:
 	bool CanEditSMInstance() const { return InstanceManager->CanEditSMInstance(InstanceId); }
 	bool GetSMInstanceTransform(FTransform& OutInstanceTransform, bool bWorldSpace = false) const { return InstanceManager->GetSMInstanceTransform(InstanceId, OutInstanceTransform, bWorldSpace); }
 	bool SetSMInstanceTransform(const FTransform& InstanceTransform, bool bWorldSpace = false, bool bMarkRenderStateDirty = false, bool bTeleport = false) const { return InstanceManager->SetSMInstanceTransform(InstanceId, InstanceTransform, bWorldSpace, bMarkRenderStateDirty, bTeleport); }
+	void NotifySMInstanceMovementStarted() { return InstanceManager->NotifySMInstanceMovementStarted(InstanceId); }
+	void NotifySMInstanceMovementOngoing() { return InstanceManager->NotifySMInstanceMovementOngoing(InstanceId); }
+	void NotifySMInstanceMovementEnded() { return InstanceManager->NotifySMInstanceMovementEnded(InstanceId); }
+	void NotifySMInstanceSelectionChanged(const bool bIsSelected) { return InstanceManager->NotifySMInstanceSelectionChanged(InstanceId, bIsSelected); }
 	bool CanDeleteSMInstance() const { return InstanceManager->CanDeleteSMInstance(InstanceId); }
 	bool DeleteSMInstance() const { return InstanceManager->DeleteSMInstances(MakeArrayView(&InstanceId, 1)); }
 	bool CanDuplicateSMInstance() const { return InstanceManager->CanDuplicateSMInstance(InstanceId); }
diff --git a/Engine/Source/Runtime/Engine/Public/ISMPartition/ISMPartitionActor.h b/Engine/Source/Runtime/Engine/Public/ISMPartition/ISMPartitionActor.h
--- a/Engine/Source/Runtime/Engine/Public/ISMPartition/ISMPartitionActor.h
+++ b/Engine/Source/Runtime/Engine/Public/ISMPartition/ISMPartitionActor.h
@@ -68,7 +68,7 @@ private:
 
 protected:
 	//~ ISMInstanceManagerProvider interface
-	virtual TScriptInterface<ISMInstanceManager> GetSMInstanceManager(const FSMInstanceId& InstanceId) const override;
+	virtual TScriptInterface<ISMInstanceManager> GetSMInstanceManager(const FSMInstanceId& InstanceId) override;
 
 private:
 #if WITH_EDITORONLY_DATA
diff --git a/Engine/Source/Runtime/Foliage/Foliage.Build.cs b/Engine/Source/Runtime/Foliage/Foliage.Build.cs
--- a/Engine/Source/Runtime/Foliage/Foliage.Build.cs
+++ b/Engine/Source/Runtime/Foliage/Foliage.Build.cs
@@ -11,6 +11,7 @@ public class Foliage: ModuleRules
 				"Core",
 				"CoreUObject",
                 "Engine",
+				"TypedElementFramework",
 				"RenderCore",
 				"RHI"
 			}
diff --git a/Engine/Source/Runtime/Foliage/Private/FoliageActor.cpp b/Engine/Source/Runtime/Foliage/Private/FoliageActor.cpp
--- a/Engine/Source/Runtime/Foliage/Private/FoliageActor.cpp
+++ b/Engine/Source/Runtime/Foliage/Private/FoliageActor.cpp
@@ -254,7 +254,7 @@ void FFoliageActor::PostEditUndo(FFoliageInfo* InInfo, UFoliageType* FoliageType
 	UpdateActorTransforms(Info->Instances);
 }
 
-void FFoliageActor::PreMoveInstances(const TArray<int32>& InInstancesMoved)
+void FFoliageActor::PreMoveInstances(TArrayView<const int32> InInstancesMoved)
 {
 	for (int32 Index : InInstancesMoved)
 	{
@@ -265,7 +265,7 @@ void FFoliageActor::PreMoveInstances(const TArray<int32>& InInstancesMoved)
 	}
 }
 
-void FFoliageActor::PostMoveInstances(const TArray<int32>& InInstancesMoved, bool bFinished)
+void FFoliageActor::PostMoveInstances(TArrayView<const int32> InInstancesMoved, bool bFinished)
 {
 	// Copy because moving actors might remove them from ActorInstances
 	TArray<AActor*> MovedActors;
diff --git a/Engine/Source/Runtime/Foliage/Private/FoliageActor.h b/Engine/Source/Runtime/Foliage/Private/FoliageActor.h
--- a/Engine/Source/Runtime/Foliage/Private/FoliageActor.h
+++ b/Engine/Source/Runtime/Foliage/Private/FoliageActor.h
@@ -55,8 +55,8 @@ struct FFoliageActor : public FFoliageImpl
 	virtual void Refresh(bool Async, bool Force) override;
 	virtual void OnHiddenEditorViewMaskChanged(uint64 InHiddenEditorViews) override;
 	virtual void PostEditUndo(FFoliageInfo* InInfo, UFoliageType* FoliageType) override;
-	virtual void PreMoveInstances(const TArray<int32>& InInstancesMoved) override;
-	virtual void PostMoveInstances(const TArray<int32>& InInstancesMoved, bool bFinished) override;
+	virtual void PreMoveInstances(TArrayView<const int32> InInstancesMoved) override;
+	virtual void PostMoveInstances(TArrayView<const int32> InInstancesMoved, bool bFinished) override;
 	virtual void NotifyFoliageTypeChanged(UFoliageType* FoliageType, bool bSourceChanged) override;
 	virtual void Reapply(const UFoliageType* FoliageType) override;
 	AActor* Spawn(const FFoliageInstance& Instance);
diff --git a/Engine/Source/Runtime/Foliage/Private/InstancedFoliage.cpp b/Engine/Source/Runtime/Foliage/Private/InstancedFoliage.cpp
--- a/Engine/Source/Runtime/Foliage/Private/InstancedFoliage.cpp
+++ b/Engine/Source/Runtime/Foliage/Private/InstancedFoliage.cpp
@@ -49,6 +49,7 @@ InstancedFoliage.cpp: Instanced foliage implementation.
 #include "FoliageISMActor.h"
 #include "LevelUtils.h"
 #include "FoliageHelper.h"
+#include "Elements/Framework/EngineElementsLibrary.h"
 #include "Algo/Transform.h"
 #include "ActorPartition/ActorPartitionSubsystem.h"
 #include "Misc/CoreMisc.h"
@@ -75,6 +76,23 @@ const FGuid FFoliageCustomVersion::GUID(0x430C4D19, 0x71544970, 0x87699B69, 0xDF
 // Register the custom version with core
 FCustomVersionRegistration GRegisterFoliageCustomVersion(FFoliageCustomVersion::GUID, FFoliageCustomVersion::LatestVersion, TEXT("FoliageVer"));
 
+namespace FoliageElementUtil
+{
+
+static int32 GEnableFoliageInstanceElements = 0;
+static FAutoConsoleVariableRef CVarEnableFoliageInstanceElements(
+	TEXT("TypedElements.EnableFoliageInstanceElements"),
+	GEnableFoliageInstanceElements,
+	TEXT("Is support for static mesh instance elements enabled for foliage owned instances?")
+	);
+
+bool FoliageInstanceElementsEnabled()
+{
+	return GEnableFoliageInstanceElements != 0;
+}
+
+}
+
 ///
 /// FFoliageStaticMesh
 ///
@@ -1647,20 +1665,43 @@ void FFoliageStaticMesh::Reapply(const UFoliageType* FoliageType)
 
 		const bool bWasRegistered = Component->IsRegistered();
 		Component->UnregisterComponent();
-		Component->ClearInstances();
 		Component->InitPerInstanceRenderData(false);
 
 		Component->bAutoRebuildTreeOnInstanceChanges = false;
 
-		TArray<FTransform> InstanceTransforms;
-		Algo::Transform(Info->Instances, InstanceTransforms, &FFoliageInstance::GetInstanceWorldTransform);
-		Component->AddInstances(InstanceTransforms, /*bShouldReturnIndices*/false, /*bWorldSpace*/true);
+		// Update the instance transforms, attempting to re-use existing instances where possible
+		const int32 NumInstancesToUpdate = FMath::Min(Component->GetInstanceCount(), Info->Instances.Num());
+		if (Component->GetInstanceCount() < Info->Instances.Num())
+		{
+			// Add any missing instances
+			TArray<FTransform> InstanceTransforms;
+			Algo::Transform(MakeArrayView(Info->Instances.GetData() + Component->GetInstanceCount(), Info->Instances.Num() - Component->GetInstanceCount()), InstanceTransforms, &FFoliageInstance::GetInstanceWorldTransform);
+			Component->AddInstances(InstanceTransforms, /*bShouldReturnIndices*/false, /*bWorldSpace*/true);
+		}
+		else if (Component->GetInstanceCount() > Info->Instances.Num())
+		{
+			// Remove any redundant instances
+			TArray<int32> InstancesToRemove;
+			for (int32 InstanceIndex = Info->Instances.Num(); InstanceIndex < Component->GetInstanceCount(); ++InstanceIndex)
+			{
+				InstancesToRemove.Add(InstanceIndex);
+			}
+			Component->RemoveInstances(InstancesToRemove);
+		}
+		check(Component->GetInstanceCount() == Info->Instances.Num());
+		if (NumInstancesToUpdate > 0)
+		{
+			// Update any existing instances
+			TArray<FTransform> InstanceTransforms;
+			Algo::Transform(MakeArrayView(Info->Instances.GetData(), NumInstancesToUpdate), InstanceTransforms, &FFoliageInstance::GetInstanceWorldTransform);
+			Component->BatchUpdateInstancesTransforms(0, InstanceTransforms, /*bWorldSpace*/true, /*bMarkRenderStateDirty*/true, /*bTeleport*/true);
+		}
 
 		Component->bAutoRebuildTreeOnInstanceChanges = true;
 		Component->BuildTreeIfOutdated(true, true);
 
+		// Update the selection state of the instances
 		Component->ClearInstanceSelection();
-
 		if (Info->SelectedIndices.Num())
 		{
 			for (int32 i : Info->SelectedIndices)
@@ -1692,7 +1733,6 @@ FFoliageInfo::FFoliageInfo()
 	, IFA(nullptr)
 	  // This will cover cases of running in editor or as editor with -game.  We do not want this during cooking.
 	, InstanceHash(GIsEditor || IsRunningGame() ? new FFoliageInstanceHash() : nullptr)
-	, bMovingInstances(false)
 #endif
 { }
 
@@ -2049,12 +2089,12 @@ void FFoliageInfo::MoveInstances(AInstancedFoliageActor* InToIFA, const TSet<int
 	}
 }
 
-void FFoliageInfo::RemoveInstances(const TArray<int32>& InInstancesToRemove, bool RebuildFoliageTree)
+void FFoliageInfo::RemoveInstances(TArrayView<const int32> InInstancesToRemove, bool RebuildFoliageTree)
 {
 	RemoveInstancesImpl(InInstancesToRemove, RebuildFoliageTree, [](FFoliageImpl* Impl, int32 Index) { Impl->RemoveInstance(Index); });
 }
 
-void FFoliageInfo::RemoveInstancesImpl(const TArray<int32>& InInstancesToRemove, bool RebuildFoliageTree, FFoliageInfo::FRemoveImplementationFunc ImplementationFunc)
+void FFoliageInfo::RemoveInstancesImpl(TArrayView<const int32> InInstancesToRemove, bool RebuildFoliageTree, FFoliageInfo::FRemoveImplementationFunc ImplementationFunc)
 {
 	SCOPE_CYCLE_COUNTER(STAT_FoliageRemoveInstance);
 
@@ -2069,7 +2109,11 @@ void FFoliageInfo::RemoveInstancesImpl(const TArray<int32>& InInstancesToRemove,
 	Implementation->BeginUpdate();
 
 	TSet<int32> InstancesToRemove;
-	InstancesToRemove.Append(InInstancesToRemove);
+	InstancesToRemove.Reserve(InInstancesToRemove.Num());
+	for (int32 InstanceToRemove : InInstancesToRemove)
+	{
+		InstancesToRemove.Add(InstanceToRemove);
+	}
 
 	while (InstancesToRemove.Num() > 0)
 	{
@@ -2141,31 +2185,29 @@ void FFoliageInfo::RemoveInstancesImpl(const TArray<int32>& InInstancesToRemove,
 	CheckValid();
 }
 
-void FFoliageInfo::PreMoveInstances(const TArray<int32>& InInstancesToMove)
+void FFoliageInfo::PreMoveInstances(TArrayView<const int32> InInstancesMoved)
 {
-	bMovingInstances = true;
-
 	// Remove instances from the hash
-	for (TArray<int32>::TConstIterator It(InInstancesToMove); It; ++It)
+	for (int32 InstanceIndex : InInstancesMoved)
 	{
-		int32 InstanceIndex = *It;
+		checkf(!MovingInstances.Contains(InstanceIndex), TEXT("PreMoveInstances has already been called for index %d!"), InstanceIndex);
+		MovingInstances.Add(InstanceIndex);
+
 		const FFoliageInstance& Instance = Instances[InstanceIndex];
 		InstanceHash->RemoveInstance(Instance.Location, InstanceIndex);
 	}
 
-	Implementation->PreMoveInstances(InInstancesToMove);
+	Implementation->PreMoveInstances(InInstancesMoved);
 }
 
-
-void FFoliageInfo::PostUpdateInstances(const TArray<int32>& InInstancesUpdated, bool bReAddToHash, bool InUpdateSelection)
+void FFoliageInfo::PostUpdateInstances(TArrayView<const int32> InInstancesUpdated, bool bReAddToHash, bool InUpdateSelection)
 {
 	if (InInstancesUpdated.Num())
 	{
 		TSet<int32> UpdateSelectedIndices;
 		UpdateSelectedIndices.Reserve(InInstancesUpdated.Num());
-		for (TArray<int32>::TConstIterator It(InInstancesUpdated); It; ++It)
+		for (int32 InstanceIndex : InInstancesUpdated)
 		{
-			int32 InstanceIndex = *It;
 			const FFoliageInstance& Instance = Instances[InstanceIndex];
 
 			FTransform InstanceToWorld = Instance.GetInstanceWorldTransform();
@@ -2194,14 +2236,19 @@ void FFoliageInfo::PostUpdateInstances(const TArray<int32>& InInstancesUpdated,
 	}
 }
 
-void FFoliageInfo::PostMoveInstances(const TArray<int32>& InInstancesMoved, bool bFinished)
+void FFoliageInfo::PostMoveInstances(TArrayView<const int32> InInstancesMoved, bool bFinished)
 {
-	PostUpdateInstances(InInstancesMoved, true, true);
+	PostUpdateInstances(InInstancesMoved, /*bReAddToHash*/true, /*UpdateSelection*/true);
 	Implementation->PostMoveInstances(InInstancesMoved, bFinished);
-	bMovingInstances = false;
+
+	for (int32 InstanceIndex : InInstancesMoved)
+	{
+		checkf(MovingInstances.Contains(InstanceIndex), TEXT("PostMoveInstances has already been called for index %d, or a call to PreMoveInstances is missing!"), InstanceIndex);
+		MovingInstances.Remove(InstanceIndex);
+	}
 }
 
-void FFoliageInfo::DuplicateInstances(UFoliageType* InSettings, const TArray<int32>& InInstancesToDuplicate)
+void FFoliageInfo::DuplicateInstances(UFoliageType* InSettings, TArrayView<const int32> InInstancesToDuplicate)
 {
 	Implementation->BeginUpdate();
 
@@ -2395,14 +2442,17 @@ FBox FFoliageInfo::GetSelectionBoundingBox() const
 	return Implementation->GetSelectionBoundingBox(SelectedIndices);
 }
 
-void FFoliageInfo::SelectInstances(bool bSelect, TArray<int32>& InInstances)
+void FFoliageInfo::SelectInstances(bool bSelect, TArrayView<const int32> InInstances)
 {
 	if (InInstances.Num())
 	{
 		TSet<int32> ModifiedSelection;
 		ModifiedSelection.Reserve(InInstances.Num());
 		check(Implementation->IsInitialized());
-		IFA->Modify(false);
+		if (IFA)
+		{
+			IFA->Modify(false);
+		}
 		if (bSelect)
 		{
 			SelectedIndices.Reserve(InInstances.Num());
@@ -2801,7 +2851,7 @@ void AInstancedFoliageActor::MoveInstancesForMovedOwnedActors(AActor* InActor)
 	for (auto& Pair : FoliageInfos)
 	{
 		// Source of movement is the Foliage
-		if (Pair.Value->bMovingInstances)
+		if (Pair.Value->MovingInstances.Num() > 0)
 		{
 			continue;
 		}
@@ -5102,6 +5152,283 @@ float AInstancedFoliageActor::InternalTakeRadialDamage(float Damage, struct FRad
 	return Damage;
 }
 
+TScriptInterface<ISMInstanceManager> AInstancedFoliageActor::GetSMInstanceManager(const FSMInstanceId& InstanceId)
+{
+#if WITH_EDITOR
+	if (const FFoliageInfo* FoliageInfo = GetFoliageInfoForSMInstance(InstanceId))
+	{
+		if (FoliageElementUtil::FoliageInstanceElementsEnabled())
+		{
+			// Allow foliage instance editing
+			return this;
+		}
+
+		// Disable foliage instance editing
+		return nullptr;
+	}
+	// GetFoliageInfoForSMInstance currently only returns StaticMesh foliage, but we need to disable editing for ISMActor foliage until it's a supported workflow
+	for (auto& Pair : FoliageInfos)
+	{
+		if (Pair.Value->Type == EFoliageImplType::ISMActor && Pair.Value->Implementation->IsOwnedComponent(InstanceId.ISMComponent))
+		{
+			return nullptr;
+		}
+	}
+#endif
+
+	// This component isn't managed by the foliage actor
+	return InstanceId.ISMComponent;
+}
+
+bool AInstancedFoliageActor::CanEditSMInstance(const FSMInstanceId& InstanceId) const
+{
+#if WITH_EDITOR
+	const FFoliageInfo* FoliageInfo = GetFoliageInfoForSMInstance(InstanceId);
+	check(FoliageInfo);
+
+	return true;
+#else
+	return false;
+#endif
+}
+
+bool AInstancedFoliageActor::GetSMInstanceTransform(const FSMInstanceId& InstanceId, FTransform& OutInstanceTransform, bool bWorldSpace) const
+{
+#if WITH_EDITOR
+	const FFoliageInfo* FoliageInfo = GetFoliageInfoForSMInstance(InstanceId);
+	check(FoliageInfo);
+
+	const FFoliageInstance& FoliageInstance = FoliageInfo->Instances[InstanceId.InstanceIndex];
+	OutInstanceTransform = FoliageInstance.GetInstanceWorldTransform();
+	if (!bWorldSpace)
+	{
+		OutInstanceTransform = OutInstanceTransform.GetRelativeTransform(InstanceId.ISMComponent->GetComponentTransform());
+	}
+	return true;
+#else
+	return false;
+#endif
+}
+
+bool AInstancedFoliageActor::SetSMInstanceTransform(const FSMInstanceId& InstanceId, const FTransform& InstanceTransform, bool bWorldSpace, bool bMarkRenderStateDirty, bool bTeleport)
+{
+#if WITH_EDITOR
+	UFoliageType* FoliageType = nullptr;
+	FFoliageInfo* FoliageInfo = GetFoliageInfoForSMInstance(InstanceId, &FoliageType);
+	check(FoliageInfo && FoliageType);
+
+	Modify();
+
+	const bool bNeedsLocalMove = !FoliageInfo->MovingInstances.Contains(InstanceId.InstanceIndex);
+	if (bNeedsLocalMove)
+	{
+		HandleSMInstancePreMove(InstanceId, *FoliageInfo);
+	}
+
+	FTransform WorldTransform = bWorldSpace ? InstanceTransform : InstanceTransform * InstanceId.ISMComponent->GetComponentTransform();
+	FFoliageInstance& FoliageInstance = FoliageInfo->Instances[InstanceId.InstanceIndex];
+	FoliageInstance.SetInstanceWorldTransform(WorldTransform);
+	FoliageInfo->SetInstanceWorldTransform(InstanceId.InstanceIndex, WorldTransform, bTeleport);
+
+	if (bNeedsLocalMove)
+	{
+		HandleSMInstancePostMove(InstanceId, *FoliageInfo, FoliageType);
+	}
+
+	return true;
+#else
+	return false;
+#endif
+}
+
+void AInstancedFoliageActor::NotifySMInstanceMovementStarted(const FSMInstanceId& InstanceId)
+{
+#if WITH_EDITOR
+	FFoliageInfo* FoliageInfo = GetFoliageInfoForSMInstance(InstanceId);
+	check(FoliageInfo);
+
+	HandleSMInstancePreMove(InstanceId, *FoliageInfo);
+#endif
+}
+
+void AInstancedFoliageActor::NotifySMInstanceMovementOngoing(const FSMInstanceId& InstanceId)
+{
+#if WITH_EDITOR
+	FFoliageInfo* FoliageInfo = GetFoliageInfoForSMInstance(InstanceId);
+	check(FoliageInfo);
+
+	FoliageInfo->PostMoveInstances(MakeArrayView(&InstanceId.InstanceIndex, 1), /*bFinished*/false);
+	FoliageInfo->PreMoveInstances(MakeArrayView(&InstanceId.InstanceIndex, 1));
+#endif
+}
+
+void AInstancedFoliageActor::NotifySMInstanceMovementEnded(const FSMInstanceId& InstanceId)
+{
+#if WITH_EDITOR
+	UFoliageType* FoliageType = nullptr;
+	FFoliageInfo* FoliageInfo = GetFoliageInfoForSMInstance(InstanceId, &FoliageType);
+	check(FoliageInfo && FoliageType);
+
+	HandleSMInstancePostMove(InstanceId, *FoliageInfo, FoliageType);
+#endif
+}
+
+void AInstancedFoliageActor::NotifySMInstanceSelectionChanged(const FSMInstanceId& InstanceId, const bool bIsSelected)
+{
+#if WITH_EDITOR
+	FFoliageInfo* FoliageInfo = GetFoliageInfoForSMInstance(InstanceId);
+	check(FoliageInfo);
+
+	FoliageInfo->SelectInstances(bIsSelected, MakeArrayView(&InstanceId.InstanceIndex, 1));
+#endif
+}
+
+bool AInstancedFoliageActor::DeleteSMInstances(TArrayView<const FSMInstanceId> InstanceIds)
+{
+#if WITH_EDITOR
+	TMap<FFoliageInfo*, TArray<int32>> FoliageInfoToIndices;
+	for (const FSMInstanceId& InstanceId : InstanceIds)
+	{
+		FFoliageInfo* FoliageInfo = GetFoliageInfoForSMInstance(InstanceId);
+		check(FoliageInfo);
+
+		TArray<int32>& FoliageInstanceIndices = FoliageInfoToIndices.FindOrAdd(FoliageInfo);
+		FoliageInstanceIndices.Add(InstanceId.InstanceIndex);
+	}
+
+	for (const TPair<FFoliageInfo*, TArray<int32>>& FoliageInfoPair : FoliageInfoToIndices)
+	{
+		FFoliageInfo* FoliageInfo = FoliageInfoPair.Key;
+		const TArray<int32>& FoliageInstanceIndices = FoliageInfoPair.Value;
+
+		FoliageInfo->RemoveInstances(FoliageInstanceIndices, /*bRebuildFoliageTree*/true);
+	}
+
+	return true;
+#else
+	return false;
+#endif
+}
+
+bool AInstancedFoliageActor::DuplicateSMInstances(TArrayView<const FSMInstanceId> InstanceIds, TArray<FSMInstanceId>& OutNewInstanceIds)
+{
+#if WITH_EDITOR
+	struct FDuplicateFoliageInstanceData
+	{
+		UFoliageType* FoliageType = nullptr;
+		TArray<int32> InstanceIndices;
+	};
+
+	TMap<FFoliageInfo*, FDuplicateFoliageInstanceData> FoliageInfoToDuplicationData;
+	for (const FSMInstanceId& InstanceId : InstanceIds)
+	{
+		UFoliageType* FoliageType = nullptr;
+		FFoliageInfo* FoliageInfo = GetFoliageInfoForSMInstance(InstanceId, &FoliageType);
+		check(FoliageInfo && FoliageType);
+
+		FDuplicateFoliageInstanceData& DuplicationData = FoliageInfoToDuplicationData.FindOrAdd(FoliageInfo);
+		DuplicationData.FoliageType = FoliageType;
+		DuplicationData.InstanceIndices.Add(InstanceId.InstanceIndex);
+	}
+
+	OutNewInstanceIds.Reset();
+	for (const TPair<FFoliageInfo*, FDuplicateFoliageInstanceData>& FoliageInfoPair : FoliageInfoToDuplicationData)
+	{
+		FFoliageInfo* FoliageInfo = FoliageInfoPair.Key;
+		const FDuplicateFoliageInstanceData& DuplicationData = FoliageInfoPair.Value;
+
+		const int32 PrevNumInstances = FoliageInfo->Instances.Num();
+		FoliageInfo->DuplicateInstances(DuplicationData.FoliageType, DuplicationData.InstanceIndices);
+		check(FoliageInfo->Instances.Num() == PrevNumInstances + DuplicationData.InstanceIndices.Num());
+
+		OutNewInstanceIds.Reserve(OutNewInstanceIds.Num() + DuplicationData.InstanceIndices.Num());
+		for (int32 InstanceIndex = PrevNumInstances; InstanceIndex < FoliageInfo->Instances.Num(); ++InstanceIndex)
+		{
+			OutNewInstanceIds.Add(FSMInstanceId{ FoliageInfo->GetComponent(), InstanceIndex });
+		}
+	}
+
+	return true;
+#else
+	return false;
+#endif
+}
+
+#if WITH_EDITOR
+void AInstancedFoliageActor::HandleSMInstancePreMove(const FSMInstanceId& InstanceId, FFoliageInfo& FoliageInfo)
+{
+	FoliageInfo.PreMoveInstances(MakeArrayView(&InstanceId.InstanceIndex, 1));
+}
+
+void AInstancedFoliageActor::HandleSMInstancePostMove(const FSMInstanceId& InstanceId, FFoliageInfo& FoliageInfo, UFoliageType* FoliageType)
+{
+	FoliageInfo.PostMoveInstances(MakeArrayView(&InstanceId.InstanceIndex, 1), /*bFinished*/true);
+
+	// This instance may be been moved into a new world partition
+	// Verify, and re-instance any existing typed elements if required
+	{
+		const FFoliageInstance& FoliageInstance = FoliageInfo.Instances[InstanceId.InstanceIndex];
+		AInstancedFoliageActor* TargetIFA = AInstancedFoliageActor::Get(GetWorld(), /*bCreateIfNone*/true, GetLevel(), FoliageInstance.Location);
+		if (TargetIFA != this)
+		{
+			// Add the new instance first, as we need both instances available for the re-instance to work
+			// Note: TargetIFA can be null if the target is an unloaded cell, in which case we'll re-instance against a null handle and the existing instance will be removed
+			FSMInstanceId NewInstanceId;
+			if (TargetIFA)
+			{
+				TargetIFA->Modify();
+
+				FFoliageInfo* NewFoliageInfo = nullptr;
+				UFoliageType* NewFoliageType = TargetIFA->AddFoliageType(FoliageType, &NewFoliageInfo);
+
+				NewFoliageInfo->AddInstance(NewFoliageType, FoliageInstance, GetBaseComponentFromBaseId(FoliageInstance.BaseId));
+				NewFoliageInfo->Refresh(true, true);
+
+				NewInstanceId = FSMInstanceId{ NewFoliageInfo->GetComponent(), NewFoliageInfo->Instances.Num() - 1 };
+			}
+
+			// Re-instance this existing element to the new element
+			{
+				TMap<FSMInstanceId, FSMInstanceId> ReplacementSMInstanceIds;
+				ReplacementSMInstanceIds.Add(InstanceId, NewInstanceId);
+				UEngineElementsLibrary::ReplaceEditorSMInstanceElementHandles(ReplacementSMInstanceIds);
+			}
+
+			// Now remove the existing instance
+			FoliageInfo.RemoveInstances(MakeArrayView(&InstanceId.InstanceIndex, 1), /*bRebuildFoliageTree*/true);
+		}
+	}
+}
+
+FFoliageInfo* AInstancedFoliageActor::GetFoliageInfoForSMInstance(const FSMInstanceId& InstanceId, UFoliageType** OutFoliageType)
+{
+	for (auto& Pair : FoliageInfos)
+	{
+		// TODO: Support ISMActor foliage
+		// TOOD: Update GetSMInstanceManager when ISMActor is supported
+		if (Pair.Value->Type == EFoliageImplType::StaticMesh && Pair.Value->Implementation->IsOwnedComponent(InstanceId.ISMComponent))
+		{
+			if (OutFoliageType)
+			{
+				*OutFoliageType = Pair.Key;
+			}
+			return &Pair.Value.Get();
+		}
+	}
+
+	if (OutFoliageType)
+	{
+		*OutFoliageType = nullptr;
+	}
+	return nullptr;
+}
+
+const FFoliageInfo* AInstancedFoliageActor::GetFoliageInfoForSMInstance(const FSMInstanceId& InstanceId, UFoliageType** OutFoliageType) const
+{
+	return const_cast<AInstancedFoliageActor*>(this)->GetFoliageInfoForSMInstance(InstanceId, OutFoliageType);
+}
+#endif
+
 UFoliageInstancedStaticMeshComponent::UFoliageInstancedStaticMeshComponent(const FObjectInitializer& ObjectInitializer)
 	: Super(ObjectInitializer)
 	, bEnableDiscardOnLoad(false)
diff --git a/Engine/Source/Runtime/Foliage/Public/InstancedFoliage.h b/Engine/Source/Runtime/Foliage/Public/InstancedFoliage.h
--- a/Engine/Source/Runtime/Foliage/Public/InstancedFoliage.h
+++ b/Engine/Source/Runtime/Foliage/Public/InstancedFoliage.h
@@ -8,6 +8,7 @@ InstancedFoliage.h: Instanced foliage type definitions.
 #include "CoreMinimal.h"
 #include "Misc/Guid.h"
 #include "Misc/HashBuilder.h"
+#include "Containers/ArrayView.h"
 #include "FoliageInstanceBase.h"
 
 class AInstancedFoliageActor;
@@ -89,6 +90,13 @@ struct FFoliageInstance : public FFoliageInstancePlacementInfo
 		return FTransform(Rotation, Location, DrawScale3D);
 	}
 
+	void SetInstanceWorldTransform(const FTransform& Transform)
+	{
+		Location = Transform.GetTranslation();
+		Rotation = Transform.Rotator();
+		DrawScale3D = Transform.GetScale3D();
+	}
+
 	void AlignToNormal(const FVector& InNormal, float AlignMaxAngle = 0.f)
 	{
 		Flags |= FOLIAGE_AlignToNormal;
@@ -206,8 +214,8 @@ struct FFoliageImpl
 	virtual void SetInstanceWorldTransform(int32 InstanceIndex, const FTransform& Transform, bool bTeleport) = 0;
 	virtual FTransform GetInstanceWorldTransform(int32 InstanceIndex) const = 0;
 	virtual void PostUpdateInstances() {}
-	virtual void PreMoveInstances(const TArray<int32>& InInstancesMoved) {}
-	virtual void PostMoveInstances(const TArray<int32>& InInstancesMoved, bool bFinished) {}
+	virtual void PreMoveInstances(TArrayView<const int32> InInstancesMoved) {}
+	virtual void PostMoveInstances(TArrayView<const int32> InInstancesMoved, bool bFinished) {}
 	virtual bool IsOwnedComponent(const UPrimitiveComponent* PrimitiveComponent) const = 0;
 	
 	virtual void SelectAllInstances(bool bSelect) = 0;
@@ -268,7 +276,7 @@ struct FFoliageInfo
 	TSet<int32> SelectedIndices;
 
 	// Moving instances
-	bool bMovingInstances;
+	TSet<int32> MovingInstances;
 #endif
 
 	FOLIAGE_API FFoliageInfo();
@@ -309,15 +317,15 @@ struct FFoliageInfo
 	FOLIAGE_API void AddInstances(const UFoliageType* InSettings, const TArray<const FFoliageInstance*>& InNewInstances);
 	FOLIAGE_API void ReserveAdditionalInstances(const UFoliageType* InSettings, uint32 ReserveNum);
 
-	FOLIAGE_API void RemoveInstances(const TArray<int32>& InInstancesToRemove, bool RebuildFoliageTree);
+	FOLIAGE_API void RemoveInstances(TArrayView<const int32> InInstancesToRemove, bool RebuildFoliageTree);
 
 	FOLIAGE_API void MoveInstances(AInstancedFoliageActor* InToIFA, const TSet<int32>& InInstancesToMove, bool bKeepSelection);
 
 	// Apply changes in the FoliageType to the component
-	FOLIAGE_API void PreMoveInstances(const TArray<int32>& InInstancesToMove);
-	FOLIAGE_API void PostMoveInstances(const TArray<int32>& InInstancesMoved, bool bFinished = false);
-	FOLIAGE_API void PostUpdateInstances(const TArray<int32>& InInstancesUpdated, bool bReAddToHash = false, bool InUpdateSelection = false);
-	FOLIAGE_API void DuplicateInstances(UFoliageType* InSettings, const TArray<int32>& InInstancesToDuplicate);
+	FOLIAGE_API void PreMoveInstances(TArrayView<const int32> InInstancesMoved);
+	FOLIAGE_API void PostMoveInstances(TArrayView<const int32> InInstancesMoved, bool bFinished = false);
+	FOLIAGE_API void PostUpdateInstances(TArrayView<const int32>, bool bReAddToHash = false, bool InUpdateSelection = false);
+	FOLIAGE_API void DuplicateInstances(UFoliageType* InSettings, TArrayView<const int32> InInstancesToDuplicate);
 	FOLIAGE_API void GetInstancesInsideBounds(const FBox& Box, TArray<int32>& OutInstances) const;
 	FOLIAGE_API void GetInstancesInsideSphere(const FSphere& Sphere, TArray<int32>& OutInstances) const;
 	FOLIAGE_API void GetInstanceAtLocation(const FVector& Location, int32& OutInstance, bool& bOutSucess) const;
@@ -328,7 +336,7 @@ struct FFoliageInfo
 	// Destroy existing clusters and reassign all instances to new clusters
 	FOLIAGE_API void ReallocateClusters(UFoliageType* InSettings);
 
-	FOLIAGE_API void SelectInstances(bool bSelect, TArray<int32>& Instances);
+	FOLIAGE_API void SelectInstances(bool bSelect, TArrayView<const int32> Instances);
 
 	FOLIAGE_API void SelectInstances(bool bSelect);
 
@@ -370,7 +378,7 @@ private:
 	void AddInstanceImpl(const FFoliageInstance& InNewInstance, FAddImplementationFunc ImplementationFunc);
 
 	using FRemoveImplementationFunc = TFunctionRef<void(FFoliageImpl*, int32)>;
-	void RemoveInstancesImpl(const TArray<int32>& InInstancesToRemove, bool RebuildFoliageTree, FRemoveImplementationFunc ImplementationFunc);
+	void RemoveInstancesImpl(TArrayView<const int32> InInstancesToRemove, bool RebuildFoliageTree, FRemoveImplementationFunc ImplementationFunc);
 };
 
 
diff --git a/Engine/Source/Runtime/Foliage/Public/InstancedFoliageActor.h b/Engine/Source/Runtime/Foliage/Public/InstancedFoliageActor.h
--- a/Engine/Source/Runtime/Foliage/Public/InstancedFoliageActor.h
+++ b/Engine/Source/Runtime/Foliage/Public/InstancedFoliageActor.h
@@ -21,7 +21,7 @@ class UProceduralFoliageComponent;
 typedef TFunction<bool(const UPrimitiveComponent*)> FFoliageTraceFilterFunc;
 
 UCLASS(notplaceable, hidecategories = (Object, Rendering, Mobility), MinimalAPI, NotBlueprintable)
-class AInstancedFoliageActor : public AISMPartitionActor
+class AInstancedFoliageActor : public AISMPartitionActor, public ISMInstanceManager
 {
 	GENERATED_UCLASS_BODY()
 
@@ -57,8 +57,31 @@ public:
 	virtual bool IsLevelBoundsRelevant() const override { return false; }
 
 protected:
+	//~ ISMInstanceManagerProvider interface
+	virtual TScriptInterface<ISMInstanceManager> GetSMInstanceManager(const FSMInstanceId& InstanceId) override;
+
+	//~ ISMInstanceManager interface
+	virtual bool CanEditSMInstance(const FSMInstanceId& InstanceId) const override;
+	virtual bool GetSMInstanceTransform(const FSMInstanceId& InstanceId, FTransform& OutInstanceTransform, bool bWorldSpace = false) const override;
+	virtual bool SetSMInstanceTransform(const FSMInstanceId& InstanceId, const FTransform& InstanceTransform, bool bWorldSpace = false, bool bMarkRenderStateDirty = false, bool bTeleport = false) override;
+	virtual void NotifySMInstanceMovementStarted(const FSMInstanceId& InstanceId) override;
+	virtual void NotifySMInstanceMovementOngoing(const FSMInstanceId& InstanceId) override;
+	virtual void NotifySMInstanceMovementEnded(const FSMInstanceId& InstanceId) override;
+	virtual void NotifySMInstanceSelectionChanged(const FSMInstanceId& InstanceId, const bool bIsSelected) override;
+	virtual bool DeleteSMInstances(TArrayView<const FSMInstanceId> InstanceIds) override;
+	virtual bool DuplicateSMInstances(TArrayView<const FSMInstanceId> InstanceIds, TArray<FSMInstanceId>& OutNewInstanceIds) override;
+
+#if WITH_EDITOR
+	void HandleSMInstancePreMove(const FSMInstanceId& InstanceId, FFoliageInfo& FoliageInfo);
+	void HandleSMInstancePostMove(const FSMInstanceId& InstanceId, FFoliageInfo& FoliageInfo, UFoliageType* FoliageType);
+
+	FFoliageInfo* GetFoliageInfoForSMInstance(const FSMInstanceId& InstanceId, UFoliageType** OutFoliageType = nullptr);
+	const FFoliageInfo* GetFoliageInfoForSMInstance(const FSMInstanceId& InstanceId, UFoliageType** OutFoliageType = nullptr) const;
+#endif
+
 	// Default InternalTakeRadialDamage behavior finds and scales damage for the closest component which isn't appropriate for foliage.
 	virtual float InternalTakeRadialDamage(float Damage, struct FRadialDamageEvent const& RadialDamageEvent, class AController* EventInstigator, AActor* DamageCauser) override;
+
 public:
 #if WITH_EDITOR
 	FOLIAGE_API void EnterEditMode();
