Foliage: Fix focus crash (with empty Foliage)

diff --git a/Engine/Source/Runtime/Foliage/Private/InstancedFoliage.cpp b/Engine/Source/Runtime/Foliage/Private/InstancedFoliage.cpp
--- a/Engine/Source/Runtime/Foliage/Private/InstancedFoliage.cpp
+++ b/Engine/Source/Runtime/Foliage/Private/InstancedFoliage.cpp
@@ -2438,8 +2438,12 @@ void FFoliageInfo::SelectInstances(bool bSelect)
 
 FBox FFoliageInfo::GetSelectionBoundingBox() const
 {
-	check(Implementation->IsInitialized());
-	return Implementation->GetSelectionBoundingBox(SelectedIndices);
+	if (Implementation->IsInitialized())
+	{
+		return Implementation->GetSelectionBoundingBox(SelectedIndices);
+	}
+
+	return FBox(EForceInit::ForceInit);
 }
 
 void FFoliageInfo::SelectInstances(bool bSelect, TArrayView<const int32> InInstances)
