Added support for editing ISM partition actor instances

This adds a new interface, IISMPartitionInstanceManager, which is the ISM partition actor equivalent of ISMInstanceManager, but works on ISM partition client instances rather than ISMC instances. Internally APartitionActor will forward the underlying ISMC edit request to the appropriate IISMPartitionInstanceManager, if any.

An ISM partition client that wants to provide editing support for its instances must either register a manager via APartitionActor::RegisterClientInstanceManager, or a proxy that can provide a manager via APartitionActor::RegisterClientInstanceManagerProvider. If the manager/provider instance registered is a UObject, then that reference will be serialized and loaded again automatically. If the manager/provider is a normal C++ instance, then the client owner must re-register it manually post-load.

Unlike ISMC instances, which are always editable unless their owner actor disables it, ISM partition actor instances are only editable if a client registers a manager as described above. This is because only the client has enough knowledge to safely edit its instances.

ISM partition actors also required supporting a new editing feature; instance selection groups. This is due to an ISM partition client instance potentially being made up of several ISMC instances that must be treated as a single logical entity for editing. To achieve this we've added an ISMInstanceManager::ForEachSMInstanceInSelectionGroup function that can be overridden by any manager that needs to provide instance grouping, and which is implemented by APartitionActor to group together all the ISMC instances that belong to a given ISM partition client instance.

Our proof-of-concept for ISM partition actor instance editing is foliage, available when the CVar "TypedElements.EnableFoliageInstanceElements" is enabled. This takes the proof-of-concept work that was already done for ISMC foliage and moves it to a common implementation that works on FFoliageInfo, which both FFoliageStaticMesh (for ISMC foliage) and FFoliageISMActor (for ISM partition actor foliage) call into from their respective manager interfaces.

Misc changes:
 - AISMPartitionActor::SelectISMInstances no longer dirties the component, as the selection state is transient and not serialized on the component.
 - Added a PostSerialize and PostLoad function to foliage implementations, to allow fixing up the IFA and other data.
 - Fixed FFoliageISMActor failing to register its delegates post-load.
 - Fixed FFoliageISMActor crashing due to FoliageTypeActor being null post-load.
 - Fixed FFoliageStaticMesh and FFoliageISMActor leaving raw delegates bound post-destruction.
 - Added the ability for instance managers to override the display name and tooltip used for instances in the details panel.

diff --git a/Engine/Source/Editor/UnrealEd/Private/Elements/SMInstance/SMInstanceElementDetailsProxyObject.cpp b/Engine/Source/Editor/UnrealEd/Private/Elements/SMInstance/SMInstanceElementDetailsProxyObject.cpp
--- a/Engine/Source/Editor/UnrealEd/Private/Elements/SMInstance/SMInstanceElementDetailsProxyObject.cpp
+++ b/Engine/Source/Editor/UnrealEd/Private/Elements/SMInstance/SMInstanceElementDetailsProxyObject.cpp
@@ -107,38 +107,35 @@ UClass* FSMInstanceElementDetailsProxyObjectNameEditSink::GetSupportedClass() co
 
 FText FSMInstanceElementDetailsProxyObjectNameEditSink::GetObjectDisplayName(UObject* Object) const
 {
-	FText DisplayName;
 	USMInstanceElementDetailsProxyObject* ProxyObject = CastChecked<USMInstanceElementDetailsProxyObject>(Object);
-	FSMInstanceManager SMInstance = ProxyObject->GetSMInstance();
-	if (!SMInstance)
+	if (FSMInstanceManager SMInstance = ProxyObject->GetSMInstance())
 	{
-		return LOCTEXT("ProxyObjectDisplayNameUnknown", "Unknown Instanced Static Mesh");
-	}
-	TObjectPtr<UStaticMesh> StaticMesh = SMInstance.GetISMComponent()->GetStaticMesh();
-	if (StaticMesh)
-	{
-		DisplayName = FText::FromString(StaticMesh->GetName());
-	}
-	else
-	{
-		DisplayName = LOCTEXT("ProxyObjectDisplayNameUnassigned", "Static Mesh Unassigned");
+		FText DisplayName = SMInstance.GetSMInstanceDisplayName();
+		if (DisplayName.IsEmpty())
+		{
+			const UStaticMesh* StaticMesh = SMInstance.GetISMComponent()->GetStaticMesh();
+			const FText OwnerDisplayName = FText::FromString(StaticMesh ? StaticMesh->GetName() : SMInstance.GetISMComponent()->GetName());
+			DisplayName = FText::Format(LOCTEXT("ProxyObjectDisplayNameFmt", "{0} - Instance {1}"), OwnerDisplayName, SMInstance.GetInstanceId().InstanceIndex);
+		}
+		return DisplayName;
 	}
-	DisplayName = FText::Format(LOCTEXT("ProxyObjectDisplayNameFmt", "{0} - Instance {1}"), DisplayName, SMInstance.GetInstanceId().InstanceIndex);
-	return DisplayName;
+	return FText();
 }
 
 FText FSMInstanceElementDetailsProxyObjectNameEditSink::GetObjectNameTooltip(UObject* Object) const
 {
 	USMInstanceElementDetailsProxyObject* ProxyObject = CastChecked<USMInstanceElementDetailsProxyObject>(Object);
-	FSMInstanceManager SMInstance = ProxyObject->GetSMInstance();
-	if (!SMInstance)
+	if (FSMInstanceManager SMInstance = ProxyObject->GetSMInstance())
 	{
-		return LOCTEXT("ProxyObjectTooltipUnknown", "Unknown Instanced Static Mesh");;
+		FText Tooltip = SMInstance.GetSMInstanceTooltip();
+		if (Tooltip.IsEmpty())
+		{
+			const FText OwnerDisplayPath = FText::FromString(SMInstance.GetISMComponent()->GetPathName(SMInstance.GetISMComponent()->GetWorld())); // stops the path at the level of the world the object is in
+			Tooltip = FText::Format(LOCTEXT("ProxyObjectTooltipFmt", "Instance {0} on {1}"), SMInstance.GetInstanceId().InstanceIndex, OwnerDisplayPath);
+		}
+		return Tooltip;
 	}
-	return FText::Format(LOCTEXT("ProxyObjectTooltipFmt", "Instance {0} on {1}"),
-		SMInstance.GetInstanceId().InstanceIndex,
-		// stops the path at the level of the world the object is in
-		FText::FromString(SMInstance.GetISMComponent()->GetPathName(SMInstance.GetISMComponent()->GetWorld())));
+	return FText();
 }
 
 #undef LOCTEXT_NAMESPACE
diff --git a/Engine/Source/Editor/UnrealEd/Private/Elements/SMInstance/SMInstanceElementEditorSelectionInterface.cpp b/Engine/Source/Editor/UnrealEd/Private/Elements/SMInstance/SMInstanceElementEditorSelectionInterface.cpp
--- a/Engine/Source/Editor/UnrealEd/Private/Elements/SMInstance/SMInstanceElementEditorSelectionInterface.cpp
+++ b/Engine/Source/Editor/UnrealEd/Private/Elements/SMInstance/SMInstanceElementEditorSelectionInterface.cpp
@@ -46,25 +46,94 @@ bool USMInstanceElementEditorSelectionInterface::IsElementSelected(const FTypedE
 	FTypedElementListConstPtr SelectionSetPtr = InSelectionSet.GetElementList();
 	const FSMInstanceManager SMInstance = SMInstanceElementDataUtil::GetSMInstanceFromHandle(InElementHandle);
 
-	if (SelectionSetPtr && SMInstance)
+	if (SMInstance && SelectionSetPtr && SelectionSetPtr->Num() > 0)
 	{
-		if (SelectionSetPtr->Num() == 0)
+		// If any element within this group is selected, then the whole group is selected
+		bool bIsSelected = false;
+		SMInstance.ForEachSMInstanceInSelectionGroup([&InSelectionOptions, &SelectionSetPtr, &bIsSelected](FSMInstanceId GroupSMInstance)
 		{
-			return false;
-		}
+			if (FTypedElementHandle GroupSMInstanceElement = UEngineElementsLibrary::AcquireEditorSMInstanceElementHandle(GroupSMInstance, /*bAllowCreate*/false))
+			{
+				if (SelectionSetPtr->Contains(GroupSMInstanceElement))
+				{
+					bIsSelected = true;
+				}
+				else if (InSelectionOptions.AllowIndirect())
+				{
+					if (FTypedElementHandle ISMComponentElement = UEngineElementsLibrary::AcquireEditorComponentElementHandle(GroupSMInstance.ISMComponent, /*bAllowCreate*/false))
+					{
+						bIsSelected = SelectionSetPtr->Contains(ISMComponentElement);
+					}
+				}
+			}
+			return !bIsSelected;
+		});
+		return bIsSelected;
+	}
+
+	return false;
+}
 
-		if (SelectionSetPtr->Contains(InElementHandle))
+bool USMInstanceElementEditorSelectionInterface::SelectElement(const FTypedElementHandle& InElementHandle, FTypedElementListProxy InSelectionSet, const FTypedElementSelectionOptions& InSelectionOptions)
+{
+	FTypedElementListPtr SelectionSetPtr = InSelectionSet.GetElementList();
+	FSMInstanceManager SMInstance = SMInstanceElementDataUtil::GetSMInstanceFromHandle(InElementHandle);
+
+	if (SMInstance && SelectionSetPtr && SelectionSetPtr->Num() > 0)
+	{
+		// If any element within this group is already in the selection set, then we have nothing more to do
+		bool bIsSelected = false;
+		SMInstance.ForEachSMInstanceInSelectionGroup([&SelectionSetPtr, &bIsSelected](FSMInstanceId GroupSMInstance)
 		{
-			return true;
+			if (FTypedElementHandle GroupSMInstanceElement = UEngineElementsLibrary::AcquireEditorSMInstanceElementHandle(GroupSMInstance, /*bAllowCreate*/false))
+			{
+				if (SelectionSetPtr->Contains(GroupSMInstanceElement))
+				{
+					bIsSelected = true;
+				}
+			}
+			return !bIsSelected;
+		});
+		if (bIsSelected)
+		{
+			return false;
 		}
+	}
+
+	if (SMInstance && SelectionSetPtr && SelectionSetPtr->Add(InElementHandle))
+	{
+		SMInstance.NotifySMInstanceSelectionChanged(/*bIsSelected*/true);
+		return true;
+	}
 
-		if (InSelectionOptions.AllowIndirect())
+	return false;
+}
+
+bool USMInstanceElementEditorSelectionInterface::DeselectElement(const FTypedElementHandle& InElementHandle, FTypedElementListProxy InSelectionSet, const FTypedElementSelectionOptions& InSelectionOptions)
+{
+	FTypedElementListPtr SelectionSetPtr = InSelectionSet.GetElementList();
+	FSMInstanceManager SMInstance = SMInstanceElementDataUtil::GetSMInstanceFromHandle(InElementHandle);
+
+	if (SMInstance && SelectionSetPtr && SelectionSetPtr->Num() > 0)
+	{
+		bool bSelectionChanged = false;
+
+		// Deselect every element within this group
+		SMInstance.ForEachSMInstanceInSelectionGroup([&SelectionSetPtr, &bSelectionChanged](FSMInstanceId GroupSMInstance)
 		{
-			if (FTypedElementHandle ISMComponentElement = UEngineElementsLibrary::AcquireEditorComponentElementHandle(SMInstance.GetISMComponent(), /*bAllowCreate*/false))
+			if (FTypedElementHandle GroupSMInstanceElement = UEngineElementsLibrary::AcquireEditorSMInstanceElementHandle(GroupSMInstance, /*bAllowCreate*/false))
 			{
-				return SelectionSetPtr->Contains(ISMComponentElement);
+				bSelectionChanged |= SelectionSetPtr->Remove(GroupSMInstanceElement);
 			}
+			return true;
+		});
+
+		if (bSelectionChanged)
+		{
+			SMInstance.NotifySMInstanceSelectionChanged(/*bIsSelected*/false);
 		}
+
+		return bSelectionChanged;
 	}
 
 	return false;
diff --git a/Engine/Source/Editor/UnrealEd/Private/Elements/SMInstance/SMInstanceElementEditorSelectionInterface.h b/Engine/Source/Editor/UnrealEd/Private/Elements/SMInstance/SMInstanceElementEditorSelectionInterface.h
--- a/Engine/Source/Editor/UnrealEd/Private/Elements/SMInstance/SMInstanceElementEditorSelectionInterface.h
+++ b/Engine/Source/Editor/UnrealEd/Private/Elements/SMInstance/SMInstanceElementEditorSelectionInterface.h
@@ -12,6 +12,8 @@ class USMInstanceElementEditorSelectionInterface : public USMInstanceElementSele
 
 public:
 	virtual bool IsElementSelected(const FTypedElementHandle& InElementHandle, const FTypedElementListProxy InSelectionSet, const FTypedElementIsSelectedOptions& InSelectionOptions) override;
+	virtual bool SelectElement(const FTypedElementHandle& InElementHandle, FTypedElementListProxy InSelectionSet, const FTypedElementSelectionOptions& InSelectionOptions) override;
+	virtual bool DeselectElement(const FTypedElementHandle& InElementHandle, FTypedElementListProxy InSelectionSet, const FTypedElementSelectionOptions& InSelectionOptions) override;
 	virtual bool ShouldPreventTransactions(const FTypedElementHandle& InElementHandle) override;
 	virtual TUniquePtr<ITypedElementTransactedElement> CreateTransactedElementImpl() override;
 };
diff --git a/Engine/Source/Runtime/Engine/Classes/GameFramework/LightWeightInstanceStaticMeshManager.h b/Engine/Source/Runtime/Engine/Classes/GameFramework/LightWeightInstanceStaticMeshManager.h
--- a/Engine/Source/Runtime/Engine/Classes/GameFramework/LightWeightInstanceStaticMeshManager.h
+++ b/Engine/Source/Runtime/Engine/Classes/GameFramework/LightWeightInstanceStaticMeshManager.h
@@ -54,6 +54,8 @@ public:
 protected:
 
 	//~ ISMInstanceManager interface
+	virtual FText GetSMInstanceDisplayName(const FSMInstanceId& InstanceId) const override;
+	virtual FText GetSMInstanceTooltip(const FSMInstanceId& InstanceId) const override;
 	virtual bool CanEditSMInstance(const FSMInstanceId& InstanceId) const override;
 	virtual bool CanMoveSMInstance(const FSMInstanceId& InstanceId, const ETypedElementWorldType InWorldType) const override;
 	virtual bool GetSMInstanceTransform(const FSMInstanceId& InstanceId, FTransform& OutInstanceTransform, bool bWorldSpace = false) const override;
diff --git a/Engine/Source/Runtime/Engine/Private/ISMPartition/ISMPartitionActor.cpp b/Engine/Source/Runtime/Engine/Private/ISMPartition/ISMPartitionActor.cpp
--- a/Engine/Source/Runtime/Engine/Private/ISMPartition/ISMPartitionActor.cpp
+++ b/Engine/Source/Runtime/Engine/Private/ISMPartition/ISMPartitionActor.cpp
@@ -118,7 +118,9 @@ void AISMPartitionActor::UnregisterClient(FISMClientHandle& Handle)
 	Modify();
 	check(Handle.Guid == Clients[Handle.Index]);
 	RemoveISMInstances(Handle);
-			
+	
+	ClientInstanceManagers.Remove(Handle.Guid);
+
 	int32 ClientIndex = Handle.Index;
 	Clients[Handle.Index] = FGuid();
 	Handle.Index = -1;
@@ -130,6 +132,22 @@ void AISMPartitionActor::UnregisterClient(FISMClientHandle& Handle)
 	}
 }
 
+void AISMPartitionActor::RegisterClientInstanceManager(const FISMClientHandle& Handle, IISMPartitionInstanceManager* ClientInstanceManager)
+{
+	check(Handle.Guid == Clients[Handle.Index]);
+
+	FISMClientInstanceManagerData& ClientInstanceManagerData = ClientInstanceManagers.FindOrAdd(Handle.Guid);
+	ClientInstanceManagerData.SetInstanceManager(ClientInstanceManager);
+}
+
+void AISMPartitionActor::RegisterClientInstanceManagerProvider(const FISMClientHandle& Handle, IISMPartitionInstanceManagerProvider* ClientInstanceManagerProvider)
+{
+	check(Handle.Guid == Clients[Handle.Index]);
+
+	FISMClientInstanceManagerData& ClientInstanceManagerData = ClientInstanceManagers.FindOrAdd(Handle.Guid);
+	ClientInstanceManagerData.SetInstanceManagerProvider(ClientInstanceManagerProvider);
+}
+
 int32 AISMPartitionActor::RegisterISMComponentDescriptor(const FISMComponentDescriptor& Descriptor)
 {
 	check(Descriptor.Hash != 0);
@@ -466,8 +484,6 @@ void AISMPartitionActor::SelectISMInstances(const FISMClientHandle& Handle, bool
 				FISMClientData& ClientData = ComponentData.ClientInstances[Handle.Index];
 				if (ClientData.Instances.Num())
 				{
-					ModifyComponent(ComponentData);
-
 					for (int32 SelectIndex : Indices)
 					{
 						FISMClientInstance& ClientInstance = ClientData.Instances[SelectIndex];
@@ -633,26 +649,298 @@ void AISMPartitionActor::UpdateHISMTrees(bool bAsync, bool bForce)
 	}
 }
 
-void AISMPartitionActor::GetClientComponents(const FISMClientHandle& Handle, TArray<UInstancedStaticMeshComponent*>& OutComponents)
+void AISMPartitionActor::ForEachClientComponent(const FISMClientHandle& Handle, TFunctionRef<bool(UInstancedStaticMeshComponent*)> Callback) const
 {
 	check(Handle.Guid == Clients[Handle.Index]);
-	for (FISMComponentData& ComponentData : DescriptorComponents)
+	for (const FISMComponentData& ComponentData : DescriptorComponents)
 	{
 		if (ComponentData.Component)
 		{
-			FISMClientData& ClientData = ComponentData.ClientInstances[Handle.Index];
+			const FISMClientData& ClientData = ComponentData.ClientInstances[Handle.Index];
 			if (ClientData.Instances.Num())
 			{
-				OutComponents.Add(ComponentData.Component);
+				if (!Callback(ComponentData.Component))
+				{
+					return;
+				}
+			}
+		}
+	}
+}
+
+void AISMPartitionActor::ForEachClientSMInstance(const FISMClientHandle& Handle, TFunctionRef<bool(FSMInstanceId)> Callback) const
+{
+	check(Handle.Guid == Clients[Handle.Index]);
+	for (const FISMComponentData& ComponentData : DescriptorComponents)
+	{
+		if (ComponentData.Component)
+		{
+			const FISMClientData& ClientData = ComponentData.ClientInstances[Handle.Index];
+			for (const FISMClientInstance& ClientInstance : ClientData.Instances)
+			{
+				for (int32 ComponentIndex : ClientInstance.ComponentIndices)
+				{
+					if (!Callback(FSMInstanceId{ ComponentData.Component, ComponentIndex }))
+					{
+						return;
+					}
+				}
+			}
+		}
+	}
+}
+
+void AISMPartitionActor::ForEachClientSMInstance(const FISMClientHandle& Handle, int32 InstanceIndex, TFunctionRef<bool(FSMInstanceId)> Callback) const
+{
+	check(Handle.Guid == Clients[Handle.Index]);
+	for (const FISMComponentData& ComponentData : DescriptorComponents)
+	{
+		if (ComponentData.Component)
+		{
+			const FISMClientData& ClientData = ComponentData.ClientInstances[Handle.Index];
+			const FISMClientInstance& ClientInstance = ClientData.Instances[InstanceIndex];
+			for (int32 ComponentIndex : ClientInstance.ComponentIndices)
+			{
+				if (!Callback(FSMInstanceId{ ComponentData.Component, ComponentIndex }))
+				{
+					return;
+				}
 			}
 		}
 	}
 }
+#endif
+
+FText AISMPartitionActor::GetSMInstanceDisplayName(const FSMInstanceId& InstanceId) const
+{
+#if WITH_EDITOR
+	FISMPartitionInstanceManager InstanceManager = GetISMPartitionInstanceManagerChecked(InstanceId);
+	return InstanceManager.GetISMPartitionInstanceDisplayName();
+#else
+	return FText();
+#endif
+}
+
+FText AISMPartitionActor::GetSMInstanceTooltip(const FSMInstanceId& InstanceId) const
+{
+#if WITH_EDITOR
+	FISMPartitionInstanceManager InstanceManager = GetISMPartitionInstanceManagerChecked(InstanceId);
+	return InstanceManager.GetISMPartitionInstanceTooltip();
+#else
+	return FText();
+#endif
+}
+
+bool AISMPartitionActor::CanEditSMInstance(const FSMInstanceId& InstanceId) const
+{
+#if WITH_EDITOR
+	FISMPartitionInstanceManager InstanceManager = GetISMPartitionInstanceManagerChecked(InstanceId);
+	return InstanceManager.CanEditISMPartitionInstance();
+#else
+	return false;
+#endif
+}
 
+bool AISMPartitionActor::CanMoveSMInstance(const FSMInstanceId& InstanceId, const ETypedElementWorldType WorldType) const
+{
+#if WITH_EDITOR
+	FISMPartitionInstanceManager InstanceManager = GetISMPartitionInstanceManagerChecked(InstanceId);
+	return InstanceManager.CanMoveISMPartitionInstance(WorldType);
+#else
+	return false;
+#endif
+}
+
+bool AISMPartitionActor::GetSMInstanceTransform(const FSMInstanceId& InstanceId, FTransform& OutInstanceTransform, bool bWorldSpace) const
+{
+#if WITH_EDITOR
+	FISMPartitionInstanceManager InstanceManager = GetISMPartitionInstanceManagerChecked(InstanceId);
+	return InstanceManager.GetISMPartitionInstanceTransform(OutInstanceTransform, bWorldSpace);
+#else
+	return false;
+#endif
+}
+
+bool AISMPartitionActor::SetSMInstanceTransform(const FSMInstanceId& InstanceId, const FTransform& InstanceTransform, bool bWorldSpace, bool bMarkRenderStateDirty, bool bTeleport)
+{
+#if WITH_EDITOR
+	FISMPartitionInstanceManager InstanceManager = GetISMPartitionInstanceManagerChecked(InstanceId);
+	return InstanceManager.SetISMPartitionInstanceTransform(InstanceTransform, bWorldSpace, bTeleport);
+#else
+	return false;
+#endif
+}
+
+void AISMPartitionActor::NotifySMInstanceMovementStarted(const FSMInstanceId& InstanceId)
+{
+#if WITH_EDITOR
+	FISMPartitionInstanceManager InstanceManager = GetISMPartitionInstanceManagerChecked(InstanceId);
+	InstanceManager.NotifyISMPartitionInstanceMovementStarted();
+#endif
+}
+
+void AISMPartitionActor::NotifySMInstanceMovementOngoing(const FSMInstanceId& InstanceId)
+{
+#if WITH_EDITOR
+	FISMPartitionInstanceManager InstanceManager = GetISMPartitionInstanceManagerChecked(InstanceId);
+	InstanceManager.NotifyISMPartitionInstanceMovementOngoing();
+#endif
+}
+
+void AISMPartitionActor::NotifySMInstanceMovementEnded(const FSMInstanceId& InstanceId)
+{
+#if WITH_EDITOR
+	FISMPartitionInstanceManager InstanceManager = GetISMPartitionInstanceManagerChecked(InstanceId);
+	InstanceManager.NotifyISMPartitionInstanceMovementEnded();
+#endif
+}
+
+void AISMPartitionActor::NotifySMInstanceSelectionChanged(const FSMInstanceId& InstanceId, const bool bIsSelected)
+{
+#if WITH_EDITOR
+	FISMPartitionInstanceManager InstanceManager = GetISMPartitionInstanceManagerChecked(InstanceId);
+	InstanceManager.NotifyISMPartitionInstanceSelectionChanged(bIsSelected);
+#endif
+}
+
+void AISMPartitionActor::ForEachSMInstanceInSelectionGroup(const FSMInstanceId& InstanceId, TFunctionRef<bool(FSMInstanceId)> Callback)
+{
+#if WITH_EDITOR
+	FISMPartitionInstanceManager InstanceManager = GetISMPartitionInstanceManagerChecked(InstanceId);
+	ForEachClientSMInstance(InstanceManager.GetInstanceId().Handle, InstanceManager.GetInstanceId().Index, Callback);
+#else
+	Callback(InstanceId);
+#endif
+}
+
+bool AISMPartitionActor::CanDeleteSMInstance(const FSMInstanceId& InstanceId) const
+{
+#if WITH_EDITOR
+	FISMPartitionInstanceManager InstanceManager = GetISMPartitionInstanceManagerChecked(InstanceId);
+	return InstanceManager.CanDeleteISMPartitionInstance();
+#else
+	return false;
+#endif
+}
+
+bool AISMPartitionActor::DeleteSMInstances(TArrayView<const FSMInstanceId> InstanceIds)
+{
+#if WITH_EDITOR
+	// Batch by the client manager
+	TMap<IISMPartitionInstanceManager*, TArray<FISMClientInstanceId>> BatchedInstancesToDelete;
+	for (const FSMInstanceId& InstanceId : InstanceIds)
+	{
+		FISMPartitionInstanceManager InstanceManager = GetISMPartitionInstanceManagerChecked(InstanceId);
+		TArray<FISMClientInstanceId>& ClientInstanceIds = BatchedInstancesToDelete.FindOrAdd(InstanceManager.GetInstanceManager());
+		ClientInstanceIds.Add(InstanceManager.GetInstanceId());
+	}
+
+	bool bDidDelete = false;
+	for (TTuple<IISMPartitionInstanceManager*, TArray<FISMClientInstanceId>>& BatchedInstancesToDeletePair : BatchedInstancesToDelete)
+	{
+		bDidDelete |= BatchedInstancesToDeletePair.Key->DeleteISMPartitionInstances(BatchedInstancesToDeletePair.Value);
+	}
+	return bDidDelete;
+#else
+	return false;
+#endif
+}
+
+bool AISMPartitionActor::CanDuplicateSMInstance(const FSMInstanceId& InstanceId) const
+{
+#if WITH_EDITOR
+	FISMPartitionInstanceManager InstanceManager = GetISMPartitionInstanceManagerChecked(InstanceId);
+	return InstanceManager.CanDuplicateISMPartitionInstance();
+#else
+	return false;
+#endif
+}
+
+bool AISMPartitionActor::DuplicateSMInstances(TArrayView<const FSMInstanceId> InstanceIds, TArray<FSMInstanceId>& OutNewInstanceIds)
+{
+#if WITH_EDITOR
+	// Batch by the client manager
+	TMap<IISMPartitionInstanceManager*, TArray<FISMClientInstanceId>> BatchedInstancesToDuplicate;
+	for (const FSMInstanceId& InstanceId : InstanceIds)
+	{
+		FISMPartitionInstanceManager InstanceManager = GetISMPartitionInstanceManagerChecked(InstanceId);
+		TArray<FISMClientInstanceId>& ClientInstanceIds = BatchedInstancesToDuplicate.FindOrAdd(InstanceManager.GetInstanceManager());
+		ClientInstanceIds.Add(InstanceManager.GetInstanceId());
+	}
+
+	bool bDidDuplicate = false;
+	for (const TTuple<IISMPartitionInstanceManager*, TArray<FISMClientInstanceId>>& BatchedInstancesToDuplicatePair : BatchedInstancesToDuplicate)
+	{
+		TArray<FISMClientInstanceId> NewClientInstanceIds;
+		if (BatchedInstancesToDuplicatePair.Key->DuplicateISMPartitionInstances(BatchedInstancesToDuplicatePair.Value, NewClientInstanceIds))
+		{
+			bDidDuplicate = true;
+
+			OutNewInstanceIds.Reserve(OutNewInstanceIds.Num() + NewClientInstanceIds.Num());
+			for (const FISMClientInstanceId& NewClientInstanceId : NewClientInstanceIds)
+			{
+				for (const FISMComponentData& DescriptorComponent : DescriptorComponents)
+				{
+					if (DescriptorComponent.Component)
+					{
+						const FISMClientData& ClientData = DescriptorComponent.ClientInstances[NewClientInstanceId.Handle.Index];
+						if (ClientData.Instances.IsValidIndex(NewClientInstanceId.Index))
+						{
+							// We only return the first ISMC instance, as all ISMC instances belonging to this client instance are treated as one (see ForEachSMInstanceInSelectionGroup)
+							OutNewInstanceIds.Add(FSMInstanceId{ DescriptorComponent.Component, ClientData.Instances[NewClientInstanceId.Index].ComponentIndices[0] });
+							break;
+						}
+					}
+				}
+			}
+		}
+	}
+	return bDidDuplicate;
+#else
+	return false;
+#endif
+}
+
+#if WITH_EDITOR
+FISMPartitionInstanceManager AISMPartitionActor::GetISMPartitionInstanceManager(const FSMInstanceId& InstanceId) const
+{
+	if (GIsEditor)
+	{
+		for (const FISMComponentData& DescriptorComponent : DescriptorComponents)
+		{
+			if (DescriptorComponent.Component == InstanceId.ISMComponent)
+			{
+				const FISMComponentInstance& ComponentInstance = DescriptorComponent.Instances[InstanceId.InstanceIndex];
+
+				const FGuid& ClientGuid = Clients[ComponentInstance.ClientIndex];
+				if (const FISMClientInstanceManagerData* ClientInstanceManagerData = ClientInstanceManagers.Find(ClientGuid))
+				{
+					const FISMClientHandle ClientHandle(ComponentInstance.ClientIndex, ClientGuid);
+					return FISMPartitionInstanceManager(FISMClientInstanceId{ ClientHandle, ComponentInstance.InstanceIndex }, ClientInstanceManagerData->ResolveInstanceManager(ClientHandle));
+				}
+			}
+		}
+	}
+
+	return FISMPartitionInstanceManager();
+}
+
+FISMPartitionInstanceManager AISMPartitionActor::GetISMPartitionInstanceManagerChecked(const FSMInstanceId& InstanceId) const
+{
+	FISMPartitionInstanceManager InstanceManager = GetISMPartitionInstanceManager(InstanceId);
+	check(InstanceManager);
+	return InstanceManager;
+}
 #endif
 
 ISMInstanceManager* AISMPartitionActor::GetSMInstanceManager(const FSMInstanceId& InstanceId)
 {
-	// Disable per-instance editing for now...
+#if WITH_EDITOR
+	if (GetISMPartitionInstanceManager(InstanceId))
+	{
+		return this;
+	}
+#endif
+
 	return nullptr;
 }
diff --git a/Engine/Source/Runtime/Engine/Private/LightWeightInstanceStaticMeshManager.cpp b/Engine/Source/Runtime/Engine/Private/LightWeightInstanceStaticMeshManager.cpp
--- a/Engine/Source/Runtime/Engine/Private/LightWeightInstanceStaticMeshManager.cpp
+++ b/Engine/Source/Runtime/Engine/Private/LightWeightInstanceStaticMeshManager.cpp
@@ -234,6 +234,22 @@ void ALightWeightInstanceStaticMeshManager::ClearStaticMesh()
 	OnStaticMeshSet();
 }
 
+FText ALightWeightInstanceStaticMeshManager::GetSMInstanceDisplayName(const FSMInstanceId& InstanceId) const
+{
+	check(InstanceId.ISMComponent == InstancedStaticMeshComponent);
+
+	const FText OwnerDisplayName = FText::FromString(RepresentedClass ? RepresentedClass->GetName() : AActor::GetName());
+	return FText::Format(NSLOCTEXT("LightWeightInstanceStaticMeshManager", "DisplayNameFmt", "{0} - Instance {1}"), OwnerDisplayName, InstanceId.InstanceIndex);
+}
+
+FText ALightWeightInstanceStaticMeshManager::GetSMInstanceTooltip(const FSMInstanceId& InstanceId) const
+{
+	check(InstanceId.ISMComponent == InstancedStaticMeshComponent);
+	
+	const FText OwnerDisplayPath = FText::FromString(GetPathName(GetWorld())); // stops the path at the level of the world the object is in
+	return FText::Format(NSLOCTEXT("LightWeightInstanceStaticMeshManager", "TooltipFmt", "Instance {0} on {1}"), InstanceId.InstanceIndex, OwnerDisplayPath);
+}
+
 bool ALightWeightInstanceStaticMeshManager::CanEditSMInstance(const FSMInstanceId& InstanceId) const
 {
 	check(InstanceId.ISMComponent == InstancedStaticMeshComponent);
diff --git a/Engine/Source/Runtime/Engine/Public/Elements/SMInstance/SMInstanceManager.h b/Engine/Source/Runtime/Engine/Public/Elements/SMInstance/SMInstanceManager.h
--- a/Engine/Source/Runtime/Engine/Public/Elements/SMInstance/SMInstanceManager.h
+++ b/Engine/Source/Runtime/Engine/Public/Elements/SMInstance/SMInstanceManager.h
@@ -15,7 +15,7 @@ enum class ETypedElementWorldType : uint8;
  * This exists so that actors that directly manage instances can inject custom logic around the manipulation of the underlying ISM component.
  * @note The static mesh instances given to this API must be valid and belong to this instance manager. The instance manager implementation is free to assert or crash if that contract is broken.
  */
-UINTERFACE(BlueprintType, meta=(CannotImplementInterfaceInBlueprint))
+UINTERFACE()
 class ENGINE_API USMInstanceManager : public UInterface
 {
 	GENERATED_BODY()
@@ -25,6 +25,22 @@ class ENGINE_API ISMInstanceManager
 	GENERATED_BODY()
 
 public:
+	/**
+	 * Get the display name of the given static mesh instance.
+	 */
+	virtual FText GetSMInstanceDisplayName(const FSMInstanceId& InstanceId) const
+	{
+		return FText();
+	}
+
+	/**
+	 * Get the tooltip of the given static mesh instance.
+	 */
+	virtual FText GetSMInstanceTooltip(const FSMInstanceId& InstanceId) const
+	{
+		return FText();
+	}
+
 	/**
 	 * Can the given static mesh instance be edited?
 	 * @return True if it can be edited, false otherwise.
@@ -35,7 +51,7 @@ public:
 	 * Can the given static mesh instance be moved in the world?
 	 * @return True if it can be moved, false otherwise.
 	 */
-	virtual bool CanMoveSMInstance(const FSMInstanceId& InstanceId, const ETypedElementWorldType InWorldType) const = 0;
+	virtual bool CanMoveSMInstance(const FSMInstanceId& InstanceId, const ETypedElementWorldType WorldType) const = 0;
 
 	/**
 	 * Attempt to get the transform of the given static mesh instance.
@@ -75,6 +91,24 @@ public:
 	 */
 	virtual void NotifySMInstanceSelectionChanged(const FSMInstanceId& InstanceId, const bool bIsSelected) = 0;
 
+	/**
+	 * Enumerate every static mesh instance element within the selection group that the given static mesh instance belongs to (including the given static mesh instance).
+	 * A selection group allows disparate static mesh instances to be considered as a single logical unit for selection, and is mostly used when a manager creates multiple static mesh instances for a single managed instance.
+	 * 
+	 * If *any* static mesh instance within the group is selected, then *all* static mesh instances within the group are considered selected.
+	 * This means that:
+	 *	- Querying the selection state for any static mesh element within the group will be consistent, regardless of which static mesh instance is in the selection set.
+	 *	- Attempting to select a static mesh instance within the group, when another static mesh instance from the group is already selected, will be a no-op.
+	 *	- Attempting to deselect a static mesh instance within the group will deselect all static mesh instances within the group.
+	 * 
+	 * There is no guarantee of which static mesh instance within the group will actually end up in the selection set, nor which static mesh instance within the group NotifySMInstanceSelectionChanged will be called for.
+	 * It is the responsibility of the manager to handle this correctly by treating any static mesh instance within a selection group as one logical unit for all the static mesh instances within the group.
+	 */
+	virtual void ForEachSMInstanceInSelectionGroup(const FSMInstanceId& InstanceId, TFunctionRef<bool(FSMInstanceId)> Callback)
+	{
+		Callback(InstanceId);
+	}
+
 	/**
 	 * Can the given static mesh instance be deleted?
 	 * @return True if it can be deleted, false otherwise.
@@ -111,7 +145,7 @@ public:
  * This exists so that actors that indirectly manage instances can provide the correct underlying manager for the specific ISM component and instance.
  * @note The static mesh instances given to this API must be valid. The instance manager provider implementation is free to assert or crash if that contract is broken.
  */
-UINTERFACE(BlueprintType, meta=(CannotImplementInterfaceInBlueprint))
+UINTERFACE()
 class ENGINE_API USMInstanceManagerProvider : public UInterface
 {
 	GENERATED_BODY()
@@ -171,14 +205,17 @@ public:
 	int32 GetISMInstanceIndex() const { return InstanceId.InstanceIndex; }
 
 	//~ ISMInstanceManager interface
+	FText GetSMInstanceDisplayName() const { return InstanceManager->GetSMInstanceDisplayName(InstanceId); }
+	FText GetSMInstanceTooltip() const { return InstanceManager->GetSMInstanceTooltip(InstanceId); }
 	bool CanEditSMInstance() const { return InstanceManager->CanEditSMInstance(InstanceId); }
-	bool CanMoveSMInstance(const ETypedElementWorldType InWorldType) const { return InstanceManager->CanMoveSMInstance(InstanceId, InWorldType); }
+	bool CanMoveSMInstance(const ETypedElementWorldType WorldType) const { return InstanceManager->CanMoveSMInstance(InstanceId, WorldType); }
 	bool GetSMInstanceTransform(FTransform& OutInstanceTransform, bool bWorldSpace = false) const { return InstanceManager->GetSMInstanceTransform(InstanceId, OutInstanceTransform, bWorldSpace); }
 	bool SetSMInstanceTransform(const FTransform& InstanceTransform, bool bWorldSpace = false, bool bMarkRenderStateDirty = false, bool bTeleport = false) const { return InstanceManager->SetSMInstanceTransform(InstanceId, InstanceTransform, bWorldSpace, bMarkRenderStateDirty, bTeleport); }
-	void NotifySMInstanceMovementStarted() { return InstanceManager->NotifySMInstanceMovementStarted(InstanceId); }
-	void NotifySMInstanceMovementOngoing() { return InstanceManager->NotifySMInstanceMovementOngoing(InstanceId); }
-	void NotifySMInstanceMovementEnded() { return InstanceManager->NotifySMInstanceMovementEnded(InstanceId); }
-	void NotifySMInstanceSelectionChanged(const bool bIsSelected) { return InstanceManager->NotifySMInstanceSelectionChanged(InstanceId, bIsSelected); }
+	void NotifySMInstanceMovementStarted() const { return InstanceManager->NotifySMInstanceMovementStarted(InstanceId); }
+	void NotifySMInstanceMovementOngoing() const { return InstanceManager->NotifySMInstanceMovementOngoing(InstanceId); }
+	void NotifySMInstanceMovementEnded() const { return InstanceManager->NotifySMInstanceMovementEnded(InstanceId); }
+	void NotifySMInstanceSelectionChanged(const bool bIsSelected) const { return InstanceManager->NotifySMInstanceSelectionChanged(InstanceId, bIsSelected); }
+	void ForEachSMInstanceInSelectionGroup(TFunctionRef<bool(FSMInstanceId)> Callback) const { return InstanceManager->ForEachSMInstanceInSelectionGroup(InstanceId, Callback); }
 	bool CanDeleteSMInstance() const { return InstanceManager->CanDeleteSMInstance(InstanceId); }
 	bool DeleteSMInstance() const { return InstanceManager->DeleteSMInstances(MakeArrayView(&InstanceId, 1)); }
 	bool CanDuplicateSMInstance() const { return InstanceManager->CanDuplicateSMInstance(InstanceId); }
@@ -195,5 +232,5 @@ public:
 
 private:
 	FSMInstanceId InstanceId;
-	ISMInstanceManager* InstanceManager;
+	ISMInstanceManager* InstanceManager = nullptr;
 };
diff --git a/Engine/Source/Runtime/Engine/Public/ISMPartition/ISMPartitionActor.h b/Engine/Source/Runtime/Engine/Public/ISMPartition/ISMPartitionActor.h
--- a/Engine/Source/Runtime/Engine/Public/ISMPartition/ISMPartitionActor.h
+++ b/Engine/Source/Runtime/Engine/Public/ISMPartition/ISMPartitionActor.h
@@ -3,7 +3,8 @@
 
 #include "CoreMinimal.h"
 #include "ActorPartition/PartitionActor.h"
-#include "Misc/Guid.h"
+#include "ISMPartition/ISMPartitionClient.h"
+#include "ISMPartition/ISMPartitionInstanceManager.h"
 #include "ISMPartition/ISMComponentDescriptor.h"
 #include "ISMPartition/ISMComponentData.h"
 #include "Templates/Tuple.h"
@@ -11,14 +12,12 @@
 #include "Elements/SMInstance/SMInstanceManager.h"
 #include "ISMPartitionActor.generated.h"
 
-struct FISMClientHandle;
-
 DECLARE_LOG_CATEGORY_EXTERN(LogISMPartition, Log, All);
 
 /** Actor base class for instance containers placed on a grid.
 	See UActorPartitionSubsystem. */
 UCLASS(Abstract)
-class ENGINE_API AISMPartitionActor : public APartitionActor, public ISMInstanceManagerProvider
+class ENGINE_API AISMPartitionActor : public APartitionActor, public ISMInstanceManager, public ISMInstanceManagerProvider
 {
 	GENERATED_UCLASS_BODY()
 
@@ -32,6 +31,9 @@ public:
 	FISMClientHandle RegisterClient(const FGuid& ClientGuid);
 	void UnregisterClient(FISMClientHandle& ClientHandle);
 
+	void RegisterClientInstanceManager(const FISMClientHandle& Handle, IISMPartitionInstanceManager* ClientInstanceManager);
+	void RegisterClientInstanceManagerProvider(const FISMClientHandle& Handle, IISMPartitionInstanceManagerProvider* ClientInstanceManagerProvider);
+
 	int32 RegisterISMComponentDescriptor(const FISMComponentDescriptor& Descriptor);
 	const FISMComponentDescriptor& GetISMComponentDescriptor(int32 DescriptorIndex) const { return Descriptors[DescriptorIndex]; }
 
@@ -48,7 +50,9 @@ public:
 	void BeginUpdate();
 	void EndUpdate();
 	void UpdateHISMTrees(bool bAsync, bool bForce);
-	void GetClientComponents(const FISMClientHandle& Handle, TArray<UInstancedStaticMeshComponent*>& OutComponents);
+	void ForEachClientComponent(const FISMClientHandle& Handle, TFunctionRef<bool(UInstancedStaticMeshComponent*)> Callback) const;
+	void ForEachClientSMInstance(const FISMClientHandle& Handle, TFunctionRef<bool(FSMInstanceId)> Callback) const;
+	void ForEachClientSMInstance(const FISMClientHandle& Handle, int32 InstanceIndex, TFunctionRef<bool(FSMInstanceId)> Callback) const;
 	void OutputStats() const;
 
 protected:
@@ -66,9 +70,32 @@ private:
 	void ModifyActor();
 #endif
 
+private:
+	//~ ISMInstanceManager interface
+	virtual FText GetSMInstanceDisplayName(const FSMInstanceId& InstanceId) const override final;
+	virtual FText GetSMInstanceTooltip(const FSMInstanceId& InstanceId) const override final;
+	virtual bool CanEditSMInstance(const FSMInstanceId& InstanceId) const override final;
+	virtual bool CanMoveSMInstance(const FSMInstanceId& InstanceId, const ETypedElementWorldType WorldType) const override final;
+	virtual bool GetSMInstanceTransform(const FSMInstanceId& InstanceId, FTransform& OutInstanceTransform, bool bWorldSpace = false) const override final;
+	virtual bool SetSMInstanceTransform(const FSMInstanceId& InstanceId, const FTransform& InstanceTransform, bool bWorldSpace = false, bool bMarkRenderStateDirty = false, bool bTeleport = false) override final;
+	virtual void NotifySMInstanceMovementStarted(const FSMInstanceId& InstanceId) override final;
+	virtual void NotifySMInstanceMovementOngoing(const FSMInstanceId& InstanceId) override final;
+	virtual void NotifySMInstanceMovementEnded(const FSMInstanceId& InstanceId) override final;
+	virtual void NotifySMInstanceSelectionChanged(const FSMInstanceId& InstanceId, const bool bIsSelected) override final;
+	virtual void ForEachSMInstanceInSelectionGroup(const FSMInstanceId& InstanceId, TFunctionRef<bool(FSMInstanceId)> Callback) override final;
+	virtual bool CanDeleteSMInstance(const FSMInstanceId& InstanceId) const override final;
+	virtual bool DeleteSMInstances(TArrayView<const FSMInstanceId> InstanceIds) override final;
+	virtual bool CanDuplicateSMInstance(const FSMInstanceId& InstanceId) const override final;
+	virtual bool DuplicateSMInstances(TArrayView<const FSMInstanceId> InstanceIds, TArray<FSMInstanceId>& OutNewInstanceIds) override final;
+
+#if WITH_EDITOR
+	FISMPartitionInstanceManager GetISMPartitionInstanceManager(const FSMInstanceId& InstanceId) const;
+	FISMPartitionInstanceManager GetISMPartitionInstanceManagerChecked(const FSMInstanceId& InstanceId) const;
+#endif
+
 protected:
 	//~ ISMInstanceManagerProvider interface
-	virtual ISMInstanceManager* GetSMInstanceManager(const FSMInstanceId& InstanceId) override;
+	virtual ISMInstanceManager* GetSMInstanceManager(const FSMInstanceId& InstanceId) override /*final*/; // Note: This should also be final and private, but AInstancedFoliageActor needs to override it to support ISMC foliage
 
 private:
 #if WITH_EDITORONLY_DATA
@@ -81,6 +108,9 @@ private:
 	UPROPERTY()
 	TArray<FISMComponentData> DescriptorComponents;
 
+	UPROPERTY()
+	TMap<FGuid, FISMClientInstanceManagerData> ClientInstanceManagers;
+
 	/** If greater than 0 means we are between a BeginUpdate/EndUpdate call and there are some things we can delay/optimize */
 	int32 UpdateDepth;
 
@@ -88,39 +118,3 @@ private:
 	bool bWasModifyCalled;
 #endif
 };
-
-USTRUCT()
-struct FISMClientHandle
-{
-	GENERATED_USTRUCT_BODY();
-
-	FISMClientHandle()
-		: Index(-1)
-	{
-
-	}
-
-	bool IsValid() const { return Index >= 0 && Guid.IsValid(); }
-
-	void Serialize(FArchive& Ar)
-	{
-		Ar << Index;
-		Ar << Guid;
-	}
-
-private:
-	friend class AISMPartitionActor;
-
-	FISMClientHandle(int32 ClientIndex, FGuid ClientGuid)
-		: Index(ClientIndex)
-		, Guid(ClientGuid)
-	{
-
-	}
-
-	UPROPERTY()
-	int32 Index;
-
-	UPROPERTY()
-	FGuid Guid;
-};
\ No newline at end of file
diff --git a/Engine/Source/Runtime/Engine/Public/ISMPartition/ISMPartitionClient.h b/Engine/Source/Runtime/Engine/Public/ISMPartition/ISMPartitionClient.h
new file mode 100644
--- /dev/null
+++ b/Engine/Source/Runtime/Engine/Public/ISMPartition/ISMPartitionClient.h
@@ -0,0 +1,87 @@
+// Copyright Epic Games, Inc. All Rights Reserved.
+#pragma once
+
+#include "CoreMinimal.h"
+#include "Misc/Guid.h"
+#include "ISMPartitionClient.generated.h"
+
+USTRUCT()
+struct FISMClientHandle
+{
+	GENERATED_BODY();
+
+	FISMClientHandle() = default;
+
+	bool IsValid() const { return Index >= 0 && Guid.IsValid(); }
+
+	explicit operator bool() const
+	{
+		return IsValid();
+	}
+
+	bool operator==(const FISMClientHandle& InRHS) const
+	{
+		return Index == InRHS.Index
+			&& Guid == InRHS.Guid;
+	}
+
+	bool operator!=(const FISMClientHandle& InRHS) const
+	{
+		return !(*this == InRHS);
+	}
+
+	void Serialize(FArchive& Ar)
+	{
+		Ar << Index;
+		Ar << Guid;
+	}
+
+	friend inline uint32 GetTypeHash(const FISMClientHandle& InHandle)
+	{
+		return GetTypeHash(InHandle.Guid);
+	}
+
+private:
+	friend class AISMPartitionActor;
+
+	FISMClientHandle(int32 ClientIndex, FGuid ClientGuid)
+		: Index(ClientIndex)
+		, Guid(ClientGuid)
+	{
+
+	}
+
+	UPROPERTY()
+	int32 Index = INDEX_NONE;
+
+	UPROPERTY()
+	FGuid Guid;
+};
+
+struct FISMClientInstanceId
+{
+	explicit operator bool() const
+	{
+		return Handle.IsValid()
+			&& Index != INDEX_NONE;
+	}
+
+	bool operator==(const FISMClientInstanceId& InRHS) const
+	{
+		return Handle == InRHS.Handle
+			&& Index == InRHS.Index;
+	}
+
+	bool operator!=(const FISMClientInstanceId& InRHS) const
+	{
+		return !(*this == InRHS);
+	}
+
+	friend inline uint32 GetTypeHash(const FISMClientInstanceId& InId)
+	{
+		return HashCombine(GetTypeHash(InId.Index), GetTypeHash(InId.Handle));
+	}
+
+	FISMClientHandle Handle;
+	int32 Index = INDEX_NONE;
+};
diff --git a/Engine/Source/Runtime/Engine/Public/ISMPartition/ISMPartitionInstanceManager.h b/Engine/Source/Runtime/Engine/Public/ISMPartition/ISMPartitionInstanceManager.h
new file mode 100644
--- /dev/null
+++ b/Engine/Source/Runtime/Engine/Public/ISMPartition/ISMPartitionInstanceManager.h
@@ -0,0 +1,320 @@
+// Copyright Epic Games, Inc. All Rights Reserved.
+
+#pragma once
+
+#include "CoreMinimal.h"
+#include "UObject/Interface.h"
+#include "Containers/ArrayView.h"
+#include "ISMPartition/ISMPartitionClient.h"
+#include "ISMPartitionInstanceManager.generated.h"
+
+enum class ETypedElementWorldType : uint8;
+
+/**
+ * An interface for clients that manage ISM instances within a partition actor.
+ * This exists so that clients can inject custom logic around the manipulation of the underlying ISM instance.
+ * @note The client ISM instances given to this API must be valid and belong to this instance manager. The instance manager implementation is free to assert or crash if that contract is broken.
+ */
+UINTERFACE()
+class ENGINE_API UISMPartitionInstanceManager : public UInterface
+{
+	GENERATED_BODY()
+};
+class ENGINE_API IISMPartitionInstanceManager
+{
+	GENERATED_BODY()
+
+public:
+	/**
+	 * Get the display name of the given client ISM instance.
+	 */
+	virtual FText GetISMPartitionInstanceDisplayName(const FISMClientInstanceId& InstanceId) const
+	{
+		return FText();
+	}
+
+	/**
+	 * Get the tooltip of the given client ISM instance.
+	 */
+	virtual FText GetISMPartitionInstanceTooltip(const FISMClientInstanceId& InstanceId) const
+	{
+		return FText();
+	}
+
+	/**
+	 * Can the given client ISM instance be edited?
+	 * @return True if it can be edited, false otherwise.
+	 */
+	virtual bool CanEditISMPartitionInstance(const FISMClientInstanceId& InstanceId) const = 0;
+
+	/**
+	 * Can the given client ISM instance be moved in the world?
+	 * @return True if it can be moved, false otherwise.
+	 */
+	virtual bool CanMoveISMPartitionInstance(const FISMClientInstanceId& InstanceId, const ETypedElementWorldType WorldType) const = 0;
+
+	/**
+	 * Attempt to get the transform of the given client ISM instance.
+	 * @note The transform should be in the local space of the owner partition actor, unless bWorldSpace is set.
+	 * @return True if the transform was retrieved, false otherwise.
+	 */
+	virtual bool GetISMPartitionInstanceTransform(const FISMClientInstanceId& InstanceId, FTransform& OutInstanceTransform, bool bWorldSpace = false) const = 0;
+
+	/**
+	 * Attempt to set the world transform of the given client ISM instance.
+	 * @note The transform should be in the local space of the owner partition actor, unless bWorldSpace is set.
+	 * @return True if the transform was updated, false otherwise.
+	 */
+	virtual bool SetISMPartitionInstanceTransform(const FISMClientInstanceId& InstanceId, const FTransform& InstanceTransform, bool bWorldSpace = false, bool bTeleport = false) = 0;
+
+	/**
+	 * Notify that the given client ISM instance is about to be moved.
+	 * @note This gives the manager a chance to start a move operation, to avoid performing repeated work until the move is finished.
+	 */
+	virtual void NotifyISMPartitionInstanceMovementStarted(const FISMClientInstanceId& InstanceId) = 0;
+
+	/**
+	 * Notify that the given client ISM instance is currently being moved.
+	 * @note This gives the manager a chance to update a move operation.
+	 */
+	virtual void NotifyISMPartitionInstanceMovementOngoing(const FISMClientInstanceId& InstanceId) = 0;
+
+	/**
+	 * Notify that the given client ISM instance is done being moved.
+	 * @note This gives the manager a chance to end a move operation.
+	 */
+	virtual void NotifyISMPartitionInstanceMovementEnded(const FISMClientInstanceId& InstanceId) = 0;
+
+	/**
+	 * Notify that the given client ISM instance selection state has changed.
+	 * @note This gives the manager a chance to sync any internal selection state.
+	 */
+	virtual void NotifyISMPartitionInstanceSelectionChanged(const FISMClientInstanceId& InstanceId, const bool bIsSelected) = 0;
+
+	/**
+	 * Can the given client ISM instance be deleted?
+	 * @return True if it can be deleted, false otherwise.
+	 */
+	virtual bool CanDeleteISMPartitionInstance(const FISMClientInstanceId& InstanceId) const
+	{
+		return CanEditISMPartitionInstance(InstanceId);
+	}
+
+	/**
+	 * Attempt to delete the given client ISM instances.
+	 * @return True if any instances were deleted, false otherwise.
+	 */
+	virtual bool DeleteISMPartitionInstances(TArrayView<const FISMClientInstanceId> InstanceIds) = 0;
+
+	/**
+	 * Can the given client ISM instance be duplicated?
+	 * @return True if it can be duplicated, false otherwise.
+	 */
+	virtual bool CanDuplicateISMPartitionInstance(const FISMClientInstanceId& InstanceId) const
+	{
+		return CanEditISMPartitionInstance(InstanceId);
+	}
+
+	/**
+	 * Attempt to duplicate the given client ISM instances, retrieving the IDs of any new instances.
+	 * @return True if any instances were duplicated, false otherwise.
+	 */
+	virtual bool DuplicateISMPartitionInstances(TArrayView<const FISMClientInstanceId> InstanceIds, TArray<FISMClientInstanceId>& OutNewInstanceIds) = 0;
+};
+
+/**
+ * An interface for actors that can provide a manager for ISM instances within a partition actor.
+ * This exists so that objects that indirectly manage clients can provide the correct underlying manager for the specific client.
+ */
+UINTERFACE()
+class ENGINE_API UISMPartitionInstanceManagerProvider : public UInterface
+{
+	GENERATED_BODY()
+};
+class ENGINE_API IISMPartitionInstanceManagerProvider
+{
+	GENERATED_BODY()
+
+public:
+	/**
+	 * Attempt to get the instance manager associated with the given client, if any.
+	 * @return The instance manager, or null if there is no instance manager associated with the given client.
+	 */
+	virtual IISMPartitionInstanceManager* GetISMPartitionInstanceManager(const FISMClientHandle& ClientHandle) = 0;
+};
+
+/**
+ * An ISM partition instance manager, tied to a given client instance ID.
+ */
+struct ENGINE_API FISMPartitionInstanceManager
+{
+public:
+	FISMPartitionInstanceManager() = default;
+
+	FISMPartitionInstanceManager(const FISMClientInstanceId& InInstanceId, IISMPartitionInstanceManager* InInstanceManager)
+		: InstanceId(InInstanceId)
+		, InstanceManager(InInstanceManager)
+	{
+	}
+
+	explicit operator bool() const
+	{
+		return InstanceId
+			&& InstanceManager;
+	}
+
+	bool operator==(const FISMPartitionInstanceManager& InRHS) const
+	{
+		return InstanceId == InRHS.InstanceId
+			&& InstanceManager == InRHS.InstanceManager;
+	}
+
+	bool operator!=(const FISMPartitionInstanceManager& InRHS) const
+	{
+		return !(*this == InRHS);
+	}
+
+	friend inline uint32 GetTypeHash(const FISMPartitionInstanceManager& InId)
+	{
+		return GetTypeHash(InId.InstanceId);
+	}
+
+	const FISMClientInstanceId& GetInstanceId() const { return InstanceId; }
+	IISMPartitionInstanceManager* GetInstanceManager() const { return InstanceManager; }
+
+	//~ IISMPartitionInstanceManager interface
+	FText GetISMPartitionInstanceDisplayName() const { return InstanceManager->GetISMPartitionInstanceDisplayName(InstanceId); }
+	FText GetISMPartitionInstanceTooltip() const { return InstanceManager->GetISMPartitionInstanceTooltip(InstanceId); }
+	bool CanEditISMPartitionInstance() const { return InstanceManager->CanEditISMPartitionInstance(InstanceId); }
+	bool CanMoveISMPartitionInstance(const ETypedElementWorldType WorldType) const { return InstanceManager->CanMoveISMPartitionInstance(InstanceId, WorldType); }
+	bool GetISMPartitionInstanceTransform(FTransform& OutInstanceTransform, bool bWorldSpace = false) const { return InstanceManager->GetISMPartitionInstanceTransform(InstanceId, OutInstanceTransform, bWorldSpace); }
+	bool SetISMPartitionInstanceTransform(const FTransform& InstanceTransform, bool bWorldSpace = false, bool bTeleport = false) const { return InstanceManager->SetISMPartitionInstanceTransform(InstanceId, InstanceTransform, bWorldSpace, bTeleport); }
+	void NotifyISMPartitionInstanceMovementStarted() const { return InstanceManager->NotifyISMPartitionInstanceMovementStarted(InstanceId); }
+	void NotifyISMPartitionInstanceMovementOngoing() const { return InstanceManager->NotifyISMPartitionInstanceMovementOngoing(InstanceId); }
+	void NotifyISMPartitionInstanceMovementEnded() const { return InstanceManager->NotifyISMPartitionInstanceMovementEnded(InstanceId); }
+	void NotifyISMPartitionInstanceSelectionChanged(const bool bIsSelected) const { return InstanceManager->NotifyISMPartitionInstanceSelectionChanged(InstanceId, bIsSelected); }
+	bool CanDeleteISMPartitionInstance() const { return InstanceManager->CanDeleteISMPartitionInstance(InstanceId); }
+	bool DeleteISMPartitionInstance() const { return InstanceManager->DeleteISMPartitionInstances(MakeArrayView(&InstanceId, 1)); }
+	bool CanDuplicateISMPartitionInstance() const { return InstanceManager->CanDuplicateISMPartitionInstance(InstanceId); }
+	bool DuplicateISMPartitionInstance(FISMClientInstanceId& OutNewInstanceId) const
+	{
+		TArray<FISMClientInstanceId> NewInstanceIds;
+		if (InstanceManager->DuplicateISMPartitionInstances(MakeArrayView(&InstanceId, 1), NewInstanceIds))
+		{
+			OutNewInstanceId = NewInstanceIds[0];
+			return true;
+		}
+		return false;
+	}
+
+private:
+	FISMClientInstanceId InstanceId;
+	IISMPartitionInstanceManager* InstanceManager = nullptr;
+};
+
+USTRUCT()
+struct FISMClientInstanceManagerData
+{
+	GENERATED_BODY()
+	
+public:
+	bool Serialize(FArchive& Ar)
+	{
+		SerializePtr(Ar, Manager);
+		SerializePtr(Ar, ManagerProvider);
+		return true;
+	}
+
+	void SetInstanceManager(IISMPartitionInstanceManager* InManager)
+	{
+		Manager = InManager;
+		ManagerProvider = nullptr;
+	}
+
+	void SetInstanceManagerProvider(IISMPartitionInstanceManagerProvider* InManagerProvider)
+	{
+		Manager = nullptr;
+		ManagerProvider = InManagerProvider;
+	}
+	
+	IISMPartitionInstanceManager* ResolveInstanceManager(const FISMClientHandle& ClientHandle) const
+	{
+		return ManagerProvider
+			? ManagerProvider->GetISMPartitionInstanceManager(ClientHandle)
+			: Manager;
+	}
+	
+private:
+	template <typename T>
+	static void SerializePtr(FArchive& Ar, T*& Ptr)
+	{
+		enum class ERefType : uint8
+		{
+			Null,
+			Object,
+			Raw,
+		};
+
+		if (Ar.IsSaving())
+		{
+			if (UObject* Obj = Cast<UObject>(Ptr))
+			{
+				ERefType RefType = ERefType::Object;
+				Ar << RefType;
+
+				Ar << Obj;
+			}
+			else if (Ptr && !Ar.IsPersistent())
+			{
+				ERefType RefType = ERefType::Raw;
+				Ar << RefType;
+
+				uint64 PtrInt = (uint64)Ptr;
+				Ar << PtrInt;
+			}
+			else
+			{
+				ERefType RefType = ERefType::Null;
+				Ar << RefType;
+			}
+		}
+		else if (Ar.IsLoading())
+		{
+			ERefType RefType;
+			Ar << RefType;
+
+			switch (RefType)
+			{
+			case ERefType::Object:
+				{
+					UObject* Obj = nullptr;
+					Ar << Obj;
+					Ptr = Cast<T>(Obj);
+				}
+				break;
+			case ERefType::Raw:
+				{
+					uint64 PtrInt = 0;
+					Ar << PtrInt;
+					Ptr = (T*)PtrInt;
+				}
+				break;
+			default:
+				Ptr = nullptr;
+				break;
+			}
+		}
+	};
+
+	IISMPartitionInstanceManager* Manager = nullptr;
+	IISMPartitionInstanceManagerProvider* ManagerProvider = nullptr;
+};
+
+template<>
+struct TStructOpsTypeTraits<FISMClientInstanceManagerData> : public TStructOpsTypeTraitsBase2<FISMClientInstanceManagerData>
+{
+	enum
+	{
+		WithSerializer = true,
+	};
+};
diff --git a/Engine/Source/Runtime/Foliage/Private/FoliageISMActor.cpp b/Engine/Source/Runtime/Foliage/Private/FoliageISMActor.cpp
--- a/Engine/Source/Runtime/Foliage/Private/FoliageISMActor.cpp
+++ b/Engine/Source/Runtime/Foliage/Private/FoliageISMActor.cpp
@@ -11,8 +11,16 @@
 #include "Components/StaticMeshComponent.h"
 #include "ISMPartition/ISMComponentDescriptor.h"
 #include "FoliageInstancedStaticMeshComponent.h"
+#include "Elements/Framework/EngineElementsLibrary.h"
 #include "Math/Box.h"
 
+FFoliageISMActor::~FFoliageISMActor()
+{
+#if WITH_EDITOR
+	UnregisterDelegates();
+#endif
+}
+
 void FFoliageISMActor::Serialize(FArchive& Ar)
 {
 #if WITH_EDITORONLY_DATA
@@ -23,6 +31,28 @@ void FFoliageISMActor::Serialize(FArchive& Ar)
 #endif
 }
 
+void FFoliageISMActor::PostSerialize(FArchive& Ar)
+{
+	FFoliageImpl::PostSerialize(Ar);
+#if WITH_EDITOR
+	if (GIsEditor && IsInitialized() && Ar.IsLoading())
+	{
+		GetIFA()->RegisterClientInstanceManager(ClientHandle, this);
+	}
+#endif
+}
+
+void FFoliageISMActor::PostLoad()
+{
+	FFoliageImpl::PostLoad();
+#if WITH_EDITOR
+	if (GIsEditor && IsInitialized())
+	{
+		RegisterDelegates();
+	}
+#endif
+}
+
 #if WITH_EDITOR
 void FFoliageISMActor::AddReferencedObjects(UObject* InThis, FReferenceCollector& Collector)
 {
@@ -76,7 +106,7 @@ void FFoliageISMActor::Initialize(const UFoliageType* FoliageType)
 
 	AInstancedFoliageActor* IFA = GetIFA();
 
-	FoliageTypeActor = Cast<const UFoliageType_Actor>(FoliageType);
+	const UFoliageType_Actor* FoliageTypeActor = Cast<const UFoliageType_Actor>(FoliageType);
 	ActorClass = FoliageTypeActor->ActorClass ? FoliageTypeActor->ActorClass.Get() : AActor::StaticClass();
 	
 	FActorSpawnParameters SpawnParams;
@@ -92,6 +122,7 @@ void FFoliageISMActor::Initialize(const UFoliageType* FoliageType)
 	SpawnedActor->GetComponents<UStaticMeshComponent>(StaticMeshComponents);
 
 	ClientHandle = IFA->RegisterClient(Guid);
+	IFA->RegisterClientInstanceManager(ClientHandle, this);
 
 	for (UStaticMeshComponent* StaticMeshComponent : StaticMeshComponents)
 	{
@@ -159,13 +190,35 @@ void FFoliageISMActor::UnregisterDelegates()
 
 void FFoliageISMActor::OnBlueprintChanged(UBlueprint* InBlueprint)
 {
-	Reapply(FoliageTypeActor);
+	const UFoliageType_Actor* FoliageTypeActor = Cast<UFoliageType_Actor>(GetIFA()->GetFoliageTypeForInfo(Info));
+	if (FoliageTypeActor)
+	{
+		Reapply(FoliageTypeActor);
+	}
+	else if (IsInitialized())
+	{
+		Uninitialize();
+	}
 }
 
 void FFoliageISMActor::Reapply(const UFoliageType* FoliageType)
 {
 	if (IsInitialized())
 	{
+		{
+			// We can't meaningfully re-instance the old static mesh instances to the new within this function, as)
+			//	1) The old instances are destroyed prior to creating the new ones, so they are no longer available to map
+			//	2) The new instances may not be a 1:1 mapping to the old, as they may now be a completely different structure
+			// Instead we just re-instance the old static mesh instances to nothing, so that anything referencing them is cleaned-up
+			TMap<FSMInstanceId, FSMInstanceId> ReplacementSMInstanceIds;
+			GetIFA()->ForEachClientSMInstance(ClientHandle, [&ReplacementSMInstanceIds](FSMInstanceId OldSMInstanceId)
+			{
+				ReplacementSMInstanceIds.Add(OldSMInstanceId, FSMInstanceId());
+				return true;
+			});
+			UEngineElementsLibrary::ReplaceEditorSMInstanceElementHandles(ReplacementSMInstanceIds);
+		}
+
 		Uninitialize();
 	}
 	Initialize(FoliageType);
@@ -249,21 +302,27 @@ void FFoliageISMActor::OnHiddenEditorViewMaskChanged(uint64 InHiddenEditorViews)
 	}
 
 	// This can give weird results if toggling the visibility of 2 foliage types that share the same meshes. The last one wins for now.
-	TArray<UInstancedStaticMeshComponent*> ClientComponents;
-	GetIFA()->GetClientComponents(ClientHandle, ClientComponents);
-	for (UInstancedStaticMeshComponent* Component : ClientComponents)
+	GetIFA()->ForEachClientComponent(ClientHandle, [InHiddenEditorViews](UInstancedStaticMeshComponent* Component)
 	{
 		if (UFoliageInstancedStaticMeshComponent* FoliageComponent = Cast<UFoliageInstancedStaticMeshComponent>(Component))
 		{
 			FoliageComponent->FoliageHiddenEditorViews = InHiddenEditorViews;
 			FoliageComponent->MarkRenderStateDirty();
 		}
-	}
+		return true;
+	});
+}
+
+void FFoliageISMActor::PreEditUndo(UFoliageType* FoliageType)
+{
+	FFoliageImpl::PreEditUndo(FoliageType);
+	UnregisterDelegates();
 }
 
 void FFoliageISMActor::PostEditUndo(FFoliageInfo* InInfo, UFoliageType* FoliageType)
 {
 	FFoliageImpl::PostEditUndo(InInfo, FoliageType);
+	RegisterDelegates();
 }
 
 void FFoliageISMActor::NotifyFoliageTypeWillChange(UFoliageType* FoliageType)
@@ -376,4 +435,163 @@ void FFoliageISMActor::ClearSelection(const TSet<int32>& SelectedIndices)
 	SelectAllInstances(false);
 }
 
+void FFoliageISMActor::ForEachSMInstance(TFunctionRef<bool(FSMInstanceId)> Callback) const
+{
+	if (ClientHandle)
+	{
+		GetIFA()->ForEachClientSMInstance(ClientHandle, Callback);
+	}
+}
+
+void FFoliageISMActor::ForEachSMInstance(int32 InstanceIndex, TFunctionRef<bool(FSMInstanceId)> Callback) const
+{
+	if (ClientHandle)
+	{
+		GetIFA()->ForEachClientSMInstance(ClientHandle, InstanceIndex, Callback);
+	}
+}
+
 #endif // WITH_EDITOR
+
+FText FFoliageISMActor::GetISMPartitionInstanceDisplayName(const FISMClientInstanceId& InstanceId) const
+{
+#if WITH_EDITOR
+	const FFoliageInstanceId FoliageInstanceId = ISMClientInstanceIdToFoliageInstanceId(InstanceId);
+	const FText OwnerDisplayName = FText::FromString(ActorClass ? ActorClass->GetName() : GetIFA()->GetName());
+	return FText::Format(NSLOCTEXT("FoliageISMActor", "DisplayNameFmt", "{0} - Instance {1}"), OwnerDisplayName, FoliageInstanceId.Index);
+#else
+	return FText();
+#endif
+}
+
+FText FFoliageISMActor::GetISMPartitionInstanceTooltip(const FISMClientInstanceId& InstanceId) const
+{
+#if WITH_EDITOR
+	const FFoliageInstanceId FoliageInstanceId = ISMClientInstanceIdToFoliageInstanceId(InstanceId);
+	const FText OwnerDisplayPath = FText::FromString(GetIFA()->GetPathName(GetIFA()->GetWorld())); // stops the path at the level of the world the object is in
+	return FText::Format(NSLOCTEXT("FoliageISMActor", "TooltipFmt", "Instance {0} on {1}"), FoliageInstanceId.Index, OwnerDisplayPath);
+#else
+	return FText();
+#endif
+}
+
+bool FFoliageISMActor::CanEditISMPartitionInstance(const FISMClientInstanceId& InstanceId) const
+{
+#if WITH_EDITOR
+	const FFoliageInstanceId FoliageInstanceId = ISMClientInstanceIdToFoliageInstanceId(InstanceId);
+	return GetIFA()->CanEditFoliageInstance(FoliageInstanceId);
+#else
+	return false;
+#endif
+}
+
+bool FFoliageISMActor::CanMoveISMPartitionInstance(const FISMClientInstanceId& InstanceId, const ETypedElementWorldType InWorldType) const
+{
+#if WITH_EDITOR
+	const FFoliageInstanceId FoliageInstanceId = ISMClientInstanceIdToFoliageInstanceId(InstanceId);
+	return GetIFA()->CanMoveFoliageInstance(FoliageInstanceId, InWorldType);
+#else
+	return false;
+#endif
+}
+
+bool FFoliageISMActor::GetISMPartitionInstanceTransform(const FISMClientInstanceId& InstanceId, FTransform& OutInstanceTransform, bool bWorldSpace) const
+{
+#if WITH_EDITOR
+	const FFoliageInstanceId FoliageInstanceId = ISMClientInstanceIdToFoliageInstanceId(InstanceId);
+	return GetIFA()->GetFoliageInstanceTransform(FoliageInstanceId, OutInstanceTransform, bWorldSpace);
+#else
+	return false;
+#endif
+}
+
+bool FFoliageISMActor::SetISMPartitionInstanceTransform(const FISMClientInstanceId& InstanceId, const FTransform& InstanceTransform, bool bWorldSpace, bool bTeleport)
+{
+#if WITH_EDITOR
+	const FFoliageInstanceId FoliageInstanceId = ISMClientInstanceIdToFoliageInstanceId(InstanceId);
+	return GetIFA()->SetFoliageInstanceTransform(FoliageInstanceId, InstanceTransform, bWorldSpace, bTeleport);
+#else
+	return false;
+#endif
+}
+
+void FFoliageISMActor::NotifyISMPartitionInstanceMovementStarted(const FISMClientInstanceId& InstanceId)
+{
+#if WITH_EDITOR
+	const FFoliageInstanceId FoliageInstanceId = ISMClientInstanceIdToFoliageInstanceId(InstanceId);
+	GetIFA()->NotifyFoliageInstanceMovementStarted(FoliageInstanceId);
+#endif
+}
+
+void FFoliageISMActor::NotifyISMPartitionInstanceMovementOngoing(const FISMClientInstanceId& InstanceId)
+{
+#if WITH_EDITOR
+	const FFoliageInstanceId FoliageInstanceId = ISMClientInstanceIdToFoliageInstanceId(InstanceId);
+	GetIFA()->NotifyFoliageInstanceMovementOngoing(FoliageInstanceId);
+#endif
+}
+
+void FFoliageISMActor::NotifyISMPartitionInstanceMovementEnded(const FISMClientInstanceId& InstanceId)
+{
+#if WITH_EDITOR
+	const FFoliageInstanceId FoliageInstanceId = ISMClientInstanceIdToFoliageInstanceId(InstanceId);
+	GetIFA()->NotifyFoliageInstanceMovementEnded(FoliageInstanceId);
+#endif
+}
+
+void FFoliageISMActor::NotifyISMPartitionInstanceSelectionChanged(const FISMClientInstanceId& InstanceId, const bool bIsSelected)
+{
+#if WITH_EDITOR
+	const FFoliageInstanceId FoliageInstanceId = ISMClientInstanceIdToFoliageInstanceId(InstanceId);
+	GetIFA()->NotifyFoliageInstanceSelectionChanged(FoliageInstanceId, bIsSelected);
+#endif
+}
+
+bool FFoliageISMActor::DeleteISMPartitionInstances(TArrayView<const FISMClientInstanceId> InstanceIds)
+{
+#if WITH_EDITOR
+	const TArray<FFoliageInstanceId> FoliageInstanceIds = ISMClientInstanceIdsToFoliageInstanceIds(InstanceIds);
+	return GetIFA()->DeleteFoliageInstances(FoliageInstanceIds);
+#else
+	return false;
+#endif
+}
+
+bool FFoliageISMActor::DuplicateISMPartitionInstances(TArrayView<const FISMClientInstanceId> InstanceIds, TArray<FISMClientInstanceId>& OutNewInstanceIds)
+{
+#if WITH_EDITOR
+	const TArray<FFoliageInstanceId> FoliageInstanceIds = ISMClientInstanceIdsToFoliageInstanceIds(InstanceIds);
+
+	TArray<FFoliageInstanceId> NewFoliageInstanceIds;
+	const bool bDidDuplicate = GetIFA()->DuplicateFoliageInstances(FoliageInstanceIds, NewFoliageInstanceIds);
+
+	OutNewInstanceIds.Reset(NewFoliageInstanceIds.Num());
+	for (const FFoliageInstanceId& NewFoliageInstanceId : NewFoliageInstanceIds)
+	{
+		OutNewInstanceIds.Add(FISMClientInstanceId{ static_cast<FFoliageISMActor&>(*NewFoliageInstanceId.Info->Implementation).ClientHandle, NewFoliageInstanceId.Index });
+	}
+
+	return bDidDuplicate;
+#else
+	return false;
+#endif
+}
+
+#if WITH_EDITOR
+FFoliageInstanceId FFoliageISMActor::ISMClientInstanceIdToFoliageInstanceId(const FISMClientInstanceId& InstanceId) const
+{
+	check(InstanceId.Handle == ClientHandle);
+	return FFoliageInstanceId{ Info, InstanceId.Index };
+}
+
+TArray<FFoliageInstanceId> FFoliageISMActor::ISMClientInstanceIdsToFoliageInstanceIds(TArrayView<const FISMClientInstanceId> InstanceIds) const
+{
+	TArray<FFoliageInstanceId> FoliageInstanceIds;
+	FoliageInstanceIds.Reserve(InstanceIds.Num());
+	for (const FISMClientInstanceId& InstanceId : InstanceIds)
+	{
+		FoliageInstanceIds.Add(ISMClientInstanceIdToFoliageInstanceId(InstanceId));
+	}
+	return FoliageInstanceIds;
+}
+#endif
diff --git a/Engine/Source/Runtime/Foliage/Private/FoliageISMActor.h b/Engine/Source/Runtime/Foliage/Private/FoliageISMActor.h
--- a/Engine/Source/Runtime/Foliage/Private/FoliageISMActor.h
+++ b/Engine/Source/Runtime/Foliage/Private/FoliageISMActor.h
@@ -13,26 +13,28 @@ class UPrimitiveComponent;
 class UBlueprint;
 class UFoliageType_Actor;
 
-struct FFoliageISMActor : public FFoliageImpl
+struct FFoliageISMActor : public FFoliageImpl, public IISMPartitionInstanceManager
 {
 	FFoliageISMActor(FFoliageInfo* Info)
 		: FFoliageImpl(Info)
 #if WITH_EDITORONLY_DATA
 		, Guid(FGuid::NewGuid())
 		, ActorClass(nullptr)
-		, FoliageTypeActor(nullptr)
 #endif
 	{
 	}
 
+	virtual ~FFoliageISMActor();
+
 #if WITH_EDITORONLY_DATA
 	FGuid Guid;
 	FISMClientHandle ClientHandle;
 	TSortedMap<int32, TArray<FTransform>> ISMDefinition;
 	UClass* ActorClass;
-	const UFoliageType_Actor* FoliageTypeActor;
 #endif
 	virtual void Serialize(FArchive& Ar) override;
+	virtual void PostSerialize(FArchive& Ar) override;
+	virtual void PostLoad() override;
 		
 #if WITH_EDITOR
 	virtual void AddReferencedObjects(UObject* InThis, FReferenceCollector& Collector) override;
@@ -56,10 +58,14 @@ struct FFoliageISMActor : public FFoliageImpl
 	virtual void ApplySelection(bool bApply, const TSet<int32>& SelectedIndices) override;
 	virtual void ClearSelection(const TSet<int32>& SelectedIndices) override;
 
+	virtual void ForEachSMInstance(TFunctionRef<bool(FSMInstanceId)> Callback) const override;
+	virtual void ForEachSMInstance(int32 InstanceIndex, TFunctionRef<bool(FSMInstanceId)> Callback) const override;
+
 	virtual void BeginUpdate() override;
 	virtual void EndUpdate() override;
 	virtual void Refresh(bool bAsync, bool bForce) override;
 	virtual void OnHiddenEditorViewMaskChanged(uint64 InHiddenEditorViews) override;
+	virtual void PreEditUndo(UFoliageType* FoliageType) override;
 	virtual void PostEditUndo(FFoliageInfo* InInfo, UFoliageType* FoliageType) override;
 	virtual void NotifyFoliageTypeWillChange(UFoliageType* FoliageType) override;
 	virtual void NotifyFoliageTypeChanged(UFoliageType* FoliageType, bool bSourceChanged) override;
@@ -69,4 +75,24 @@ private:
 	void UnregisterDelegates();
 	void OnBlueprintChanged(UBlueprint* InBlueprint);
 #endif
-};
\ No newline at end of file
+
+private:
+	//~ IISMPartitionInstanceManager interface
+	virtual FText GetISMPartitionInstanceDisplayName(const FISMClientInstanceId& InstanceId) const override;
+	virtual FText GetISMPartitionInstanceTooltip(const FISMClientInstanceId& InstanceId) const override;
+	virtual bool CanEditISMPartitionInstance(const FISMClientInstanceId& InstanceId) const override;
+	virtual bool CanMoveISMPartitionInstance(const FISMClientInstanceId& InstanceId, const ETypedElementWorldType InWorldType) const override;
+	virtual bool GetISMPartitionInstanceTransform(const FISMClientInstanceId& InstanceId, FTransform& OutInstanceTransform, bool bWorldSpace = false) const override;
+	virtual bool SetISMPartitionInstanceTransform(const FISMClientInstanceId& InstanceId, const FTransform& InstanceTransform, bool bWorldSpace = false, bool bTeleport = false) override;
+	virtual void NotifyISMPartitionInstanceMovementStarted(const FISMClientInstanceId& InstanceId) override;
+	virtual void NotifyISMPartitionInstanceMovementOngoing(const FISMClientInstanceId& InstanceId) override;
+	virtual void NotifyISMPartitionInstanceMovementEnded(const FISMClientInstanceId& InstanceId) override;
+	virtual void NotifyISMPartitionInstanceSelectionChanged(const FISMClientInstanceId& InstanceId, const bool bIsSelected) override;
+	virtual bool DeleteISMPartitionInstances(TArrayView<const FISMClientInstanceId> InstanceIds) override;
+	virtual bool DuplicateISMPartitionInstances(TArrayView<const FISMClientInstanceId> InstanceIds, TArray<FISMClientInstanceId>& OutNewInstanceIds) override;
+
+#if WITH_EDITOR
+	FFoliageInstanceId ISMClientInstanceIdToFoliageInstanceId(const FISMClientInstanceId& InstanceId) const;
+	TArray<FFoliageInstanceId> ISMClientInstanceIdsToFoliageInstanceIds(TArrayView<const FISMClientInstanceId> InstanceIds) const;
+#endif
+};
diff --git a/Engine/Source/Runtime/Foliage/Private/InstancedFoliage.cpp b/Engine/Source/Runtime/Foliage/Private/InstancedFoliage.cpp
--- a/Engine/Source/Runtime/Foliage/Private/InstancedFoliage.cpp
+++ b/Engine/Source/Runtime/Foliage/Private/InstancedFoliage.cpp
@@ -50,6 +50,7 @@ InstancedFoliage.cpp: Instanced foliage implementation.
 #include "LevelUtils.h"
 #include "FoliageHelper.h"
 #include "Elements/Framework/EngineElementsLibrary.h"
+#include "Elements/SMInstance/SMInstanceManager.h"
 #include "Elements/SMInstance/SMInstanceElementData.h" // For SMInstanceElementDataUtil::SMInstanceElementsEnabled
 #include "Algo/Transform.h"
 #include "ActorPartition/ActorPartitionSubsystem.h"
@@ -97,7 +98,7 @@ bool FoliageInstanceElementsEnabled()
 ///
 /// FFoliageStaticMesh
 ///
-struct FFoliageStaticMesh : public FFoliageImpl
+struct FFoliageStaticMesh : public FFoliageImpl, public ISMInstanceManager
 {
 	FFoliageStaticMesh(FFoliageInfo* Info, UHierarchicalInstancedStaticMeshComponent* InComponent)
 		: FFoliageImpl(Info)
@@ -111,6 +112,8 @@ struct FFoliageStaticMesh : public FFoliageImpl
 
 	}
 
+	virtual ~FFoliageStaticMesh();
+
 	UHierarchicalInstancedStaticMeshComponent* Component;
 #if WITH_EDITOR
 	int32 UpdateDepth;
@@ -144,6 +147,9 @@ struct FFoliageStaticMesh : public FFoliageImpl
 	virtual void ApplySelection(bool bApply, const TSet<int32>& SelectedIndices) override;
 	virtual void ClearSelection(const TSet<int32>& SelectedIndices) override;
 
+	virtual void ForEachSMInstance(TFunctionRef<bool(FSMInstanceId)> Callback) const override;
+	virtual void ForEachSMInstance(int32 InstanceIndex, TFunctionRef<bool(FSMInstanceId)> Callback) const override;
+
 	virtual void BeginUpdate() override;
 	virtual void EndUpdate() override;
 	virtual void Refresh(bool Async, bool Force) override;
@@ -167,6 +173,24 @@ struct FFoliageStaticMesh : public FFoliageImpl
 	// Recreate the component if the FoliageType's ComponentClass doesn't match the Component's class
 	void CheckComponentClass(const UFoliageType_InstancedStaticMesh* InSettings);
 	void CreateNewComponent(const UFoliageType* InSettings);
+
+private:
+	//~ ISMInstanceManager interface
+	virtual bool CanEditSMInstance(const FSMInstanceId& InstanceId) const override;
+	virtual bool CanMoveSMInstance(const FSMInstanceId& InstanceId, const ETypedElementWorldType WorldType) const override;
+	virtual bool GetSMInstanceTransform(const FSMInstanceId& InstanceId, FTransform& OutInstanceTransform, bool bWorldSpace = false) const override;
+	virtual bool SetSMInstanceTransform(const FSMInstanceId& InstanceId, const FTransform& InstanceTransform, bool bWorldSpace = false, bool bMarkRenderStateDirty = false, bool bTeleport = false) override;
+	virtual void NotifySMInstanceMovementStarted(const FSMInstanceId& InstanceId) override;
+	virtual void NotifySMInstanceMovementOngoing(const FSMInstanceId& InstanceId) override;
+	virtual void NotifySMInstanceMovementEnded(const FSMInstanceId& InstanceId) override;
+	virtual void NotifySMInstanceSelectionChanged(const FSMInstanceId& InstanceId, const bool bIsSelected) override;
+	virtual bool DeleteSMInstances(TArrayView<const FSMInstanceId> InstanceIds) override;
+	virtual bool DuplicateSMInstances(TArrayView<const FSMInstanceId> InstanceIds, TArray<FSMInstanceId>& OutNewInstanceIds) override;
+
+#if WITH_EDITOR
+	FFoliageInstanceId SMInstanceIdToFoliageInstanceId(const FSMInstanceId& InstanceId) const;
+	TArray<FFoliageInstanceId> SMInstanceIdsToFoliageInstanceIds(TArrayView<const FSMInstanceId> InstanceIds) const;
+#endif
 };
 
 struct FFoliagePlacementUtil
@@ -979,6 +1003,19 @@ FName UFoliageType::GetDisplayFName() const
 //
 // FFoliageStaticMesh
 //
+FFoliageStaticMesh::~FFoliageStaticMesh()
+{
+#if WITH_EDITOR
+	if (Component != nullptr)
+	{
+		if (Component->GetStaticMesh() != nullptr)
+		{
+			Component->GetStaticMesh()->GetOnExtendedBoundsChanged().RemoveAll(this);
+		}
+	}
+#endif
+}
+
 void FFoliageStaticMesh::AddReferencedObjects(UObject* InThis, FReferenceCollector& Collector)
 {
 	if (Component != nullptr)
@@ -992,7 +1029,6 @@ void FFoliageStaticMesh::Serialize(FArchive& Ar)
 	Ar << Component;
 }
 
-
 int32 FFoliageStaticMesh::GetOverlappingSphereCount(const FSphere& Sphere) const
 {
 	if (Component && Component->IsTreeFullyBuilt())
@@ -1197,6 +1233,28 @@ void FFoliageStaticMesh::ClearSelection(const TSet<int32>& SelectedIndices)
 	Component->MarkRenderStateDirty();
 }
 
+void FFoliageStaticMesh::ForEachSMInstance(TFunctionRef<bool(FSMInstanceId)> Callback) const
+{
+	if (Component)
+	{
+		for (int32 InstanceIndex = 0; Info->Instances.Num(); ++InstanceIndex)
+		{
+			if (!Callback(FSMInstanceId{ Component, InstanceIndex }))
+			{
+				return;
+			}
+		}
+	}
+}
+
+void FFoliageStaticMesh::ForEachSMInstance(int32 InstanceIndex, TFunctionRef<bool(FSMInstanceId)> Callback) const
+{
+	if (Component && Info->Instances.IsValidIndex(InstanceIndex))
+	{
+		Callback(FSMInstanceId{ Component, InstanceIndex });
+	}
+}
+
 void FFoliageStaticMesh::BeginUpdate()
 {
 	if (UpdateDepth == 0)
@@ -1718,6 +1776,23 @@ void FFoliageStaticMesh::Reapply(const UFoliageType* FoliageType)
 	}
 }
 
+FFoliageInstanceId FFoliageStaticMesh::SMInstanceIdToFoliageInstanceId(const FSMInstanceId& InstanceId) const
+{
+	check(InstanceId.ISMComponent == Component);
+	return FFoliageInstanceId{ Info, InstanceId.InstanceIndex };
+}
+
+TArray<FFoliageInstanceId> FFoliageStaticMesh::SMInstanceIdsToFoliageInstanceIds(TArrayView<const FSMInstanceId> InstanceIds) const
+{
+	TArray<FFoliageInstanceId> FoliageInstanceIds;
+	FoliageInstanceIds.Reserve(InstanceIds.Num());
+	for (const FSMInstanceId& InstanceId : InstanceIds)
+	{
+		FoliageInstanceIds.Add(SMInstanceIdToFoliageInstanceId(InstanceId));
+	}
+	return FoliageInstanceIds;
+}
+
 AInstancedFoliageActor* FFoliageImpl::GetIFA() const
 {
 	return Info->IFA;
@@ -1725,6 +1800,108 @@ AInstancedFoliageActor* FFoliageImpl::GetIFA() const
 
 #endif // WITH_EDITOR
 
+bool FFoliageStaticMesh::CanEditSMInstance(const FSMInstanceId& InstanceId) const
+{
+#if WITH_EDITOR
+	const FFoliageInstanceId FoliageInstanceId = SMInstanceIdToFoliageInstanceId(InstanceId);
+	return GetIFA()->CanEditFoliageInstance(FoliageInstanceId);
+#else
+	return false;
+#endif
+}
+
+bool FFoliageStaticMesh::CanMoveSMInstance(const FSMInstanceId& InstanceId, const ETypedElementWorldType WorldType) const
+{
+#if WITH_EDITOR
+	const FFoliageInstanceId FoliageInstanceId = SMInstanceIdToFoliageInstanceId(InstanceId);
+	return GetIFA()->CanMoveFoliageInstance(FoliageInstanceId, WorldType);
+#else
+	return false;
+#endif
+}
+
+bool FFoliageStaticMesh::GetSMInstanceTransform(const FSMInstanceId& InstanceId, FTransform& OutInstanceTransform, bool bWorldSpace) const
+{
+#if WITH_EDITOR
+	const FFoliageInstanceId FoliageInstanceId = SMInstanceIdToFoliageInstanceId(InstanceId);
+	return GetIFA()->GetFoliageInstanceTransform(FoliageInstanceId, OutInstanceTransform, bWorldSpace);
+#else
+	return false;
+#endif
+}
+
+bool FFoliageStaticMesh::SetSMInstanceTransform(const FSMInstanceId& InstanceId, const FTransform& InstanceTransform, bool bWorldSpace, bool bMarkRenderStateDirty, bool bTeleport)
+{
+#if WITH_EDITOR
+	const FFoliageInstanceId FoliageInstanceId = SMInstanceIdToFoliageInstanceId(InstanceId);
+	return GetIFA()->SetFoliageInstanceTransform(FoliageInstanceId, InstanceTransform, bWorldSpace, bTeleport);
+#else
+	return false;
+#endif
+}
+
+void FFoliageStaticMesh::NotifySMInstanceMovementStarted(const FSMInstanceId& InstanceId)
+{
+#if WITH_EDITOR
+	const FFoliageInstanceId FoliageInstanceId = SMInstanceIdToFoliageInstanceId(InstanceId);
+	GetIFA()->NotifyFoliageInstanceMovementStarted(FoliageInstanceId);
+#endif
+}
+
+void FFoliageStaticMesh::NotifySMInstanceMovementOngoing(const FSMInstanceId& InstanceId)
+{
+#if WITH_EDITOR
+	const FFoliageInstanceId FoliageInstanceId = SMInstanceIdToFoliageInstanceId(InstanceId);
+	GetIFA()->NotifyFoliageInstanceMovementOngoing(FoliageInstanceId);
+#endif
+}
+
+void FFoliageStaticMesh::NotifySMInstanceMovementEnded(const FSMInstanceId& InstanceId)
+{
+#if WITH_EDITOR
+	const FFoliageInstanceId FoliageInstanceId = SMInstanceIdToFoliageInstanceId(InstanceId);
+	GetIFA()->NotifyFoliageInstanceMovementEnded(FoliageInstanceId);
+#endif
+}
+
+void FFoliageStaticMesh::NotifySMInstanceSelectionChanged(const FSMInstanceId& InstanceId, const bool bIsSelected)
+{
+#if WITH_EDITOR
+	const FFoliageInstanceId FoliageInstanceId = SMInstanceIdToFoliageInstanceId(InstanceId);
+	GetIFA()->NotifyFoliageInstanceSelectionChanged(FoliageInstanceId, bIsSelected);
+#endif
+}
+
+bool FFoliageStaticMesh::DeleteSMInstances(TArrayView<const FSMInstanceId> InstanceIds)
+{
+#if WITH_EDITOR
+	const TArray<FFoliageInstanceId> FoliageInstanceIds = SMInstanceIdsToFoliageInstanceIds(InstanceIds);
+	return GetIFA()->DeleteFoliageInstances(FoliageInstanceIds);
+#else
+	return false;
+#endif
+}
+
+bool FFoliageStaticMesh::DuplicateSMInstances(TArrayView<const FSMInstanceId> InstanceIds, TArray<FSMInstanceId>& OutNewInstanceIds)
+{
+#if WITH_EDITOR
+	const TArray<FFoliageInstanceId> FoliageInstanceIds = SMInstanceIdsToFoliageInstanceIds(InstanceIds);
+
+	TArray<FFoliageInstanceId> NewFoliageInstanceIds;
+	const bool bDidDuplicate = GetIFA()->DuplicateFoliageInstances(FoliageInstanceIds, NewFoliageInstanceIds);
+
+	OutNewInstanceIds.Reset(NewFoliageInstanceIds.Num());
+	for (const FFoliageInstanceId& NewFoliageInstanceId : NewFoliageInstanceIds)
+	{
+		OutNewInstanceIds.Add(FSMInstanceId{ NewFoliageInstanceId.Info->GetComponent(), NewFoliageInstanceId.Index });
+	}
+
+	return bDidDuplicate;
+#else
+	return false;
+#endif
+}
+
 //
 // FFoliageInfo
 //
@@ -1759,6 +1936,22 @@ void FFoliageInfo::AddReferencedObjects(UObject* InThis, FReferenceCollector& Co
 	}
 }
 
+void FFoliageInfo::PostSerialize(FArchive& Ar)
+{
+	if (Implementation)
+	{
+		Implementation->PostSerialize(Ar);
+	}
+}
+
+void FFoliageInfo::PostLoad()
+{
+	if (Implementation)
+	{
+		Implementation->PostLoad();
+	}
+}
+
 void FFoliageInfo::CreateImplementation(EFoliageImplType InType)
 {
 	check(InType != EFoliageImplType::Unknown);
@@ -1923,6 +2116,22 @@ void FFoliageInfo::CheckValid()
 #endif
 }
 
+void FFoliageInfo::ForEachSMInstance(TFunctionRef<bool(FSMInstanceId)> Callback) const
+{
+	if (Implementation)
+	{
+		Implementation->ForEachSMInstance(Callback);
+	}
+}
+
+void FFoliageInfo::ForEachSMInstance(int32 InstanceIndex, TFunctionRef<bool(FSMInstanceId)> Callback) const
+{
+	if (Implementation)
+	{
+		Implementation->ForEachSMInstance(InstanceIndex, Callback);
+	}
+}
+
 /*
 */
 void FFoliageInfo::ClearSelection()
@@ -2497,10 +2706,8 @@ void FFoliageInfo::PreEditUndo(UFoliageType* FoliageType)
 	Implementation->PreEditUndo(FoliageType);
 }
 
-void FFoliageInfo::PostEditUndo(AInstancedFoliageActor* InIFA, UFoliageType* FoliageType)
+void FFoliageInfo::PostEditUndo(UFoliageType* FoliageType)
 {
-	// Set the IFA after PostEditUndo as it is not a serialized member (will be nulled on serialization)
-	IFA = InIFA;
 	Implementation->PostEditUndo(this, FoliageType);
 		
 	// Regenerate instance hash
@@ -3718,7 +3925,7 @@ void AInstancedFoliageActor::PostEditUndo()
 	{
 		FFoliageInfo& Info = *Pair.Value;
 
-		Info.PostEditUndo(this, Pair.Key);
+		Info.PostEditUndo(Pair.Key);
 	}
 	EndUpdate();
 }
@@ -3973,6 +4180,19 @@ void AInstancedFoliageActor::Serialize(FArchive& Ar)
 			}
 		}
 	}
+	
+	for (auto& Pair : FoliageInfos)
+	{
+		FFoliageInfo& Info = *Pair.Value;
+#if WITH_EDITORONLY_DATA
+		if (Ar.IsLoading())
+		{
+			// Fix-up the IFA pointer as it's not serialized
+			Info.IFA = this;
+		}
+#endif
+		Info.PostSerialize(Ar);
+	}
 }
 
 #if WITH_EDITOR
@@ -4046,6 +4266,8 @@ void AInstancedFoliageActor::BeginDestroy()
 		FWorldDelegates::PostApplyLevelOffset.Remove(OnPostApplyLevelOffsetDelegateHandle);
 
 		FWorldDelegates::OnPostWorldInitialization.Remove(OnPostWorldInitializationDelegateHandle);
+
+		FoliageInfos.Empty();
 	}
 }
 #endif
@@ -4411,6 +4633,11 @@ void AInstancedFoliageActor::PostLoad()
 			}
 		}
 	}
+
+	for (auto& Pair : FoliageInfos)
+	{
+		Pair.Value->PostLoad();
+	}
 }
 
 #if WITH_EDITOR
@@ -5157,160 +5384,79 @@ float AInstancedFoliageActor::InternalTakeRadialDamage(float Damage, struct FRad
 	return Damage;
 }
 
-ISMInstanceManager* AInstancedFoliageActor::GetSMInstanceManager(const FSMInstanceId& InstanceId)
-{
 #if WITH_EDITOR
-	if (const FFoliageInfo* FoliageInfo = GetFoliageInfoForSMInstance(InstanceId))
-	{
-		if (FoliageElementUtil::FoliageInstanceElementsEnabled())
-		{
-			// Allow foliage instance editing
-			return this;
-		}
-
-		// Disable foliage instance editing
-		return nullptr;
-	}
-	// GetFoliageInfoForSMInstance currently only returns StaticMesh foliage, but we need to disable editing for ISMActor foliage until it's a supported workflow
-	for (auto& Pair : FoliageInfos)
-	{
-		if (Pair.Value->Type == EFoliageImplType::ISMActor && Pair.Value->Implementation->IsOwnedComponent(InstanceId.ISMComponent))
-		{
-			return nullptr;
-		}
-	}
-#endif
-
-	// This component isn't managed by the foliage actor
-	return InstanceId.ISMComponent;
-}
-
-bool AInstancedFoliageActor::CanEditSMInstance(const FSMInstanceId& InstanceId) const
+bool AInstancedFoliageActor::CanEditFoliageInstance(const FFoliageInstanceId& InstanceId) const
 {
-#if WITH_EDITOR
-	const FFoliageInfo* FoliageInfo = GetFoliageInfoForSMInstance(InstanceId);
-	check(FoliageInfo);
-
 	return true;
-#else
-	return false;
-#endif
 }
 
-bool AInstancedFoliageActor::CanMoveSMInstance(const FSMInstanceId& InstanceId, const ETypedElementWorldType InWorldType) const
+bool AInstancedFoliageActor::CanMoveFoliageInstance(const FFoliageInstanceId& InstanceId, const ETypedElementWorldType WorldType) const
 {
-#if WITH_EDITOR
-	const FFoliageInfo* FoliageInfo = GetFoliageInfoForSMInstance(InstanceId);
-	check(FoliageInfo);
-
 	return true;
-#else
-	return false;
-#endif
 }
 
-bool AInstancedFoliageActor::GetSMInstanceTransform(const FSMInstanceId& InstanceId, FTransform& OutInstanceTransform, bool bWorldSpace) const
+bool AInstancedFoliageActor::GetFoliageInstanceTransform(const FFoliageInstanceId& InstanceId, FTransform& OutInstanceTransform, bool bWorldSpace) const
 {
-#if WITH_EDITOR
-	const FFoliageInfo* FoliageInfo = GetFoliageInfoForSMInstance(InstanceId);
-	check(FoliageInfo);
-
-	const FFoliageInstance& FoliageInstance = FoliageInfo->Instances[InstanceId.InstanceIndex];
+	const FFoliageInstance& FoliageInstance = InstanceId.GetInstanceChecked();
 	OutInstanceTransform = FoliageInstance.GetInstanceWorldTransform();
 	if (!bWorldSpace)
 	{
-		OutInstanceTransform = OutInstanceTransform.GetRelativeTransform(InstanceId.ISMComponent->GetComponentTransform());
+		OutInstanceTransform = OutInstanceTransform.GetRelativeTransform(GetActorTransform());
 	}
 	return true;
-#else
-	return false;
-#endif
 }
 
-bool AInstancedFoliageActor::SetSMInstanceTransform(const FSMInstanceId& InstanceId, const FTransform& InstanceTransform, bool bWorldSpace, bool bMarkRenderStateDirty, bool bTeleport)
+bool AInstancedFoliageActor::SetFoliageInstanceTransform(const FFoliageInstanceId& InstanceId, const FTransform& InstanceTransform, bool bWorldSpace, bool bTeleport)
 {
-#if WITH_EDITOR
-	UFoliageType* FoliageType = nullptr;
-	FFoliageInfo* FoliageInfo = GetFoliageInfoForSMInstance(InstanceId, &FoliageType);
-	check(FoliageInfo && FoliageType);
-
 	Modify();
 
-	const bool bNeedsLocalMove = !FoliageInfo->MovingInstances.Contains(InstanceId.InstanceIndex);
+	const bool bNeedsLocalMove = !InstanceId.Info->MovingInstances.Contains(InstanceId.Index);
 	if (bNeedsLocalMove)
 	{
-		HandleSMInstancePreMove(InstanceId, *FoliageInfo);
+		HandleFoliageInstancePreMove(InstanceId);
 	}
 
-	FTransform WorldTransform = bWorldSpace ? InstanceTransform : InstanceTransform * InstanceId.ISMComponent->GetComponentTransform();
-	FFoliageInstance& FoliageInstance = FoliageInfo->Instances[InstanceId.InstanceIndex];
+	FTransform WorldTransform = bWorldSpace ? InstanceTransform : InstanceTransform * GetActorTransform();
+	FFoliageInstance& FoliageInstance = InstanceId.GetInstanceChecked();
 	FoliageInstance.SetInstanceWorldTransform(WorldTransform);
-	FoliageInfo->SetInstanceWorldTransform(InstanceId.InstanceIndex, WorldTransform, bTeleport);
+	InstanceId.Info->SetInstanceWorldTransform(InstanceId.Index, WorldTransform, bTeleport);
 
 	if (bNeedsLocalMove)
 	{
-		HandleSMInstancePostMove(InstanceId, *FoliageInfo, FoliageType);
+		HandleFoliageInstancePostMove(InstanceId);
 	}
 
 	return true;
-#else
-	return false;
-#endif
 }
 
-void AInstancedFoliageActor::NotifySMInstanceMovementStarted(const FSMInstanceId& InstanceId)
+void AInstancedFoliageActor::NotifyFoliageInstanceMovementStarted(const FFoliageInstanceId& InstanceId)
 {
-#if WITH_EDITOR
-	FFoliageInfo* FoliageInfo = GetFoliageInfoForSMInstance(InstanceId);
-	check(FoliageInfo);
-
-	HandleSMInstancePreMove(InstanceId, *FoliageInfo);
-#endif
+	HandleFoliageInstancePreMove(InstanceId);
 }
 
-void AInstancedFoliageActor::NotifySMInstanceMovementOngoing(const FSMInstanceId& InstanceId)
+void AInstancedFoliageActor::NotifyFoliageInstanceMovementOngoing(const FFoliageInstanceId& InstanceId)
 {
-#if WITH_EDITOR
-	FFoliageInfo* FoliageInfo = GetFoliageInfoForSMInstance(InstanceId);
-	check(FoliageInfo);
-
-	FoliageInfo->PostMoveInstances(MakeArrayView(&InstanceId.InstanceIndex, 1), /*bFinished*/false);
-	FoliageInfo->PreMoveInstances(MakeArrayView(&InstanceId.InstanceIndex, 1));
-#endif
+	InstanceId.Info->PostMoveInstances(MakeArrayView(&InstanceId.Index, 1), /*bFinished*/false);
+	InstanceId.Info->PreMoveInstances(MakeArrayView(&InstanceId.Index, 1));
 }
 
-void AInstancedFoliageActor::NotifySMInstanceMovementEnded(const FSMInstanceId& InstanceId)
+void AInstancedFoliageActor::NotifyFoliageInstanceMovementEnded(const FFoliageInstanceId& InstanceId)
 {
-#if WITH_EDITOR
-	UFoliageType* FoliageType = nullptr;
-	FFoliageInfo* FoliageInfo = GetFoliageInfoForSMInstance(InstanceId, &FoliageType);
-	check(FoliageInfo && FoliageType);
-
-	HandleSMInstancePostMove(InstanceId, *FoliageInfo, FoliageType);
-#endif
+	HandleFoliageInstancePostMove(InstanceId);
 }
 
-void AInstancedFoliageActor::NotifySMInstanceSelectionChanged(const FSMInstanceId& InstanceId, const bool bIsSelected)
+void AInstancedFoliageActor::NotifyFoliageInstanceSelectionChanged(const FFoliageInstanceId& InstanceId, const bool bIsSelected)
 {
-#if WITH_EDITOR
-	FFoliageInfo* FoliageInfo = GetFoliageInfoForSMInstance(InstanceId);
-	check(FoliageInfo);
-
-	FoliageInfo->SelectInstances(bIsSelected, MakeArrayView(&InstanceId.InstanceIndex, 1));
-#endif
+	InstanceId.Info->SelectInstances(bIsSelected, MakeArrayView(&InstanceId.Index, 1));
 }
 
-bool AInstancedFoliageActor::DeleteSMInstances(TArrayView<const FSMInstanceId> InstanceIds)
+bool AInstancedFoliageActor::DeleteFoliageInstances(TArrayView<const FFoliageInstanceId> InstanceIds)
 {
-#if WITH_EDITOR
 	TMap<FFoliageInfo*, TArray<int32>> FoliageInfoToIndices;
-	for (const FSMInstanceId& InstanceId : InstanceIds)
+	for (const FFoliageInstanceId& InstanceId : InstanceIds)
 	{
-		FFoliageInfo* FoliageInfo = GetFoliageInfoForSMInstance(InstanceId);
-		check(FoliageInfo);
-
-		TArray<int32>& FoliageInstanceIndices = FoliageInfoToIndices.FindOrAdd(FoliageInfo);
-		FoliageInstanceIndices.Add(InstanceId.InstanceIndex);
+		TArray<int32>& FoliageInstanceIndices = FoliageInfoToIndices.FindOrAdd(InstanceId.Info);
+		FoliageInstanceIndices.Add(InstanceId.Index);
 	}
 
 	for (const TPair<FFoliageInfo*, TArray<int32>>& FoliageInfoPair : FoliageInfoToIndices)
@@ -5322,14 +5468,10 @@ bool AInstancedFoliageActor::DeleteSMInstances(TArrayView<const FSMInstanceId> I
 	}
 
 	return true;
-#else
-	return false;
-#endif
 }
 
-bool AInstancedFoliageActor::DuplicateSMInstances(TArrayView<const FSMInstanceId> InstanceIds, TArray<FSMInstanceId>& OutNewInstanceIds)
+bool AInstancedFoliageActor::DuplicateFoliageInstances(TArrayView<const FFoliageInstanceId> InstanceIds, TArray<FFoliageInstanceId>& OutNewInstanceIds)
 {
-#if WITH_EDITOR
 	struct FDuplicateFoliageInstanceData
 	{
 		UFoliageType* FoliageType = nullptr;
@@ -5337,18 +5479,17 @@ bool AInstancedFoliageActor::DuplicateSMInstances(TArrayView<const FSMInstanceId
 	};
 
 	TMap<FFoliageInfo*, FDuplicateFoliageInstanceData> FoliageInfoToDuplicationData;
-	for (const FSMInstanceId& InstanceId : InstanceIds)
+	for (const FFoliageInstanceId& InstanceId : InstanceIds)
 	{
-		UFoliageType* FoliageType = nullptr;
-		FFoliageInfo* FoliageInfo = GetFoliageInfoForSMInstance(InstanceId, &FoliageType);
-		check(FoliageInfo && FoliageType);
+		UFoliageType* FoliageType = GetFoliageTypeForInfo(InstanceId.Info);
+		check(FoliageType);
 
-		FDuplicateFoliageInstanceData& DuplicationData = FoliageInfoToDuplicationData.FindOrAdd(FoliageInfo);
+		FDuplicateFoliageInstanceData& DuplicationData = FoliageInfoToDuplicationData.FindOrAdd(InstanceId.Info);
 		DuplicationData.FoliageType = FoliageType;
-		DuplicationData.InstanceIndices.Add(InstanceId.InstanceIndex);
+		DuplicationData.InstanceIndices.Add(InstanceId.Index);
 	}
 
-	OutNewInstanceIds.Reset();
+	OutNewInstanceIds.Reset(InstanceIds.Num());
 	for (const TPair<FFoliageInfo*, FDuplicateFoliageInstanceData>& FoliageInfoPair : FoliageInfoToDuplicationData)
 	{
 		FFoliageInfo* FoliageInfo = FoliageInfoPair.Key;
@@ -5361,91 +5502,106 @@ bool AInstancedFoliageActor::DuplicateSMInstances(TArrayView<const FSMInstanceId
 		OutNewInstanceIds.Reserve(OutNewInstanceIds.Num() + DuplicationData.InstanceIndices.Num());
 		for (int32 InstanceIndex = PrevNumInstances; InstanceIndex < FoliageInfo->Instances.Num(); ++InstanceIndex)
 		{
-			OutNewInstanceIds.Add(FSMInstanceId{ FoliageInfo->GetComponent(), InstanceIndex });
+			OutNewInstanceIds.Add(FFoliageInstanceId{ FoliageInfo, InstanceIndex });
 		}
 	}
 
 	return true;
-#else
-	return false;
-#endif
 }
 
-#if WITH_EDITOR
-void AInstancedFoliageActor::HandleSMInstancePreMove(const FSMInstanceId& InstanceId, FFoliageInfo& FoliageInfo)
+void AInstancedFoliageActor::HandleFoliageInstancePreMove(const FFoliageInstanceId& InstanceId)
 {
-	FoliageInfo.PreMoveInstances(MakeArrayView(&InstanceId.InstanceIndex, 1));
+	InstanceId.Info->PreMoveInstances(MakeArrayView(&InstanceId.Index, 1));
 }
 
-void AInstancedFoliageActor::HandleSMInstancePostMove(const FSMInstanceId& InstanceId, FFoliageInfo& FoliageInfo, UFoliageType* FoliageType)
+void AInstancedFoliageActor::HandleFoliageInstancePostMove(const FFoliageInstanceId& InstanceId)
 {
-	FoliageInfo.PostMoveInstances(MakeArrayView(&InstanceId.InstanceIndex, 1), /*bFinished*/true);
+	InstanceId.Info->PostMoveInstances(MakeArrayView(&InstanceId.Index, 1), /*bFinished*/true);
 
 	// This instance may be been moved into a new world partition
 	// Verify, and re-instance any existing typed elements if required
 	{
-		const FFoliageInstance& FoliageInstance = FoliageInfo.Instances[InstanceId.InstanceIndex];
+		const FFoliageInstance& FoliageInstance = InstanceId.GetInstanceChecked();
 		AInstancedFoliageActor* TargetIFA = AInstancedFoliageActor::Get(GetWorld(), /*bCreateIfNone*/true, GetLevel(), FoliageInstance.Location);
 		if (TargetIFA != this)
 		{
 			// Add the new instance first, as we need both instances available for the re-instance to work
 			// Note: TargetIFA can be null if the target is an unloaded cell, in which case we'll re-instance against a null handle and the existing instance will be removed
-			FSMInstanceId NewInstanceId;
+			TArray<FSMInstanceId> NewSMInstanceIds;
 			if (TargetIFA)
 			{
 				TargetIFA->Modify();
 
+				UFoliageType* FoliageType = GetFoliageTypeForInfo(InstanceId.Info);
+				check(FoliageType);
+
 				FFoliageInfo* NewFoliageInfo = nullptr;
 				UFoliageType* NewFoliageType = TargetIFA->AddFoliageType(FoliageType, &NewFoliageInfo);
 
 				NewFoliageInfo->AddInstance(NewFoliageType, FoliageInstance, GetBaseComponentFromBaseId(FoliageInstance.BaseId));
 				NewFoliageInfo->Refresh(true, true);
 
-				NewInstanceId = FSMInstanceId{ NewFoliageInfo->GetComponent(), NewFoliageInfo->Instances.Num() - 1 };
+				NewFoliageInfo->ForEachSMInstance(NewFoliageInfo->Instances.Num() - 1, [&NewSMInstanceIds](FSMInstanceId NewSMInstanceId)
+				{
+					NewSMInstanceIds.Add(NewSMInstanceId);
+					return true;
+				});
 			}
 
 			// Re-instance this existing element to the new element
 			{
 				TMap<FSMInstanceId, FSMInstanceId> ReplacementSMInstanceIds;
-				ReplacementSMInstanceIds.Add(InstanceId, NewInstanceId);
+				{
+					int32 SMInstanceIdIndex = 0;
+					InstanceId.Info->ForEachSMInstance(InstanceId.Index, [&NewSMInstanceIds, &ReplacementSMInstanceIds, &SMInstanceIdIndex](FSMInstanceId OldSMInstanceId)
+					{
+						ReplacementSMInstanceIds.Add(OldSMInstanceId, NewSMInstanceIds.IsValidIndex(SMInstanceIdIndex) ? NewSMInstanceIds[SMInstanceIdIndex] : FSMInstanceId());
+						++SMInstanceIdIndex;
+						return true;
+					});
+				}
 				UEngineElementsLibrary::ReplaceEditorSMInstanceElementHandles(ReplacementSMInstanceIds);
 			}
 
 			// Now remove the existing instance
-			FoliageInfo.RemoveInstances(MakeArrayView(&InstanceId.InstanceIndex, 1), /*bRebuildFoliageTree*/true);
+			InstanceId.Info->RemoveInstances(MakeArrayView(&InstanceId.Index, 1), /*bRebuildFoliageTree*/true);
 		}
 	}
 }
 
-FFoliageInfo* AInstancedFoliageActor::GetFoliageInfoForSMInstance(const FSMInstanceId& InstanceId, UFoliageType** OutFoliageType)
+UFoliageType* AInstancedFoliageActor::GetFoliageTypeForInfo(const FFoliageInfo* FoliageInfo) const
 {
 	for (auto& Pair : FoliageInfos)
 	{
-		// TODO: Support ISMActor foliage
-		// TOOD: Update GetSMInstanceManager when ISMActor is supported
-		if (Pair.Value->Type == EFoliageImplType::StaticMesh && Pair.Value->Implementation->IsOwnedComponent(InstanceId.ISMComponent))
+		if (&Pair.Value.Get() == FoliageInfo)
 		{
-			if (OutFoliageType)
-			{
-				*OutFoliageType = Pair.Key;
-			}
-			return &Pair.Value.Get();
+			return Pair.Key;
 		}
 	}
-
-	if (OutFoliageType)
-	{
-		*OutFoliageType = nullptr;
-	}
 	return nullptr;
 }
+#endif
 
-const FFoliageInfo* AInstancedFoliageActor::GetFoliageInfoForSMInstance(const FSMInstanceId& InstanceId, UFoliageType** OutFoliageType) const
+ISMInstanceManager* AInstancedFoliageActor::GetSMInstanceManager(const FSMInstanceId& InstanceId)
 {
-	return const_cast<AInstancedFoliageActor*>(this)->GetFoliageInfoForSMInstance(InstanceId, OutFoliageType);
-}
+	if (!GIsEditor || !FoliageElementUtil::FoliageInstanceElementsEnabled())
+	{
+		return nullptr;
+	}
+
+#if WITH_EDITOR
+	for (auto& Pair : FoliageInfos)
+	{
+		if (Pair.Value->Type == EFoliageImplType::StaticMesh && Pair.Value->Implementation->IsOwnedComponent(InstanceId.ISMComponent))
+		{
+			return static_cast<FFoliageStaticMesh*>(Pair.Value->Implementation.Get());
+		}
+	}
 #endif
 
+	return Super::GetSMInstanceManager(InstanceId);
+}
+
 UFoliageInstancedStaticMeshComponent::UFoliageInstancedStaticMeshComponent(const FObjectInitializer& ObjectInitializer)
 	: Super(ObjectInitializer)
 	, bEnableDiscardOnLoad(false)
diff --git a/Engine/Source/Runtime/Foliage/Public/InstancedFoliage.h b/Engine/Source/Runtime/Foliage/Public/InstancedFoliage.h
--- a/Engine/Source/Runtime/Foliage/Public/InstancedFoliage.h
+++ b/Engine/Source/Runtime/Foliage/Public/InstancedFoliage.h
@@ -19,6 +19,7 @@ class UHierarchicalInstancedStaticMeshComponent;
 class UFoliageType_InstancedStaticMesh;
 class UPrimitiveComponent;
 class UStaticMesh;
+struct FSMInstanceId;
 struct FFoliageInstanceHash;
 
 DECLARE_LOG_CATEGORY_EXTERN(LogInstancedFoliage, Log, All);
@@ -199,6 +200,8 @@ struct FFoliageImpl
 
 	virtual void AddReferencedObjects(UObject* InThis, FReferenceCollector& Collector) {}
 	virtual void Serialize(FArchive& Ar) = 0;
+	virtual void PostSerialize(FArchive& Ar) {}
+	virtual void PostLoad() {}
 
 #if WITH_EDITORONLY_DATA
 	// Not serialized but FFoliageInfo will make sure it stays valid (mostly Undo/Redo)
@@ -231,6 +234,9 @@ struct FFoliageImpl
 	virtual void ApplySelection(bool bApply, const TSet<int32>& SelectedIndices) = 0;
 	virtual void ClearSelection(const TSet<int32>& SelectedIndices) = 0;
 
+	virtual void ForEachSMInstance(TFunctionRef<bool(FSMInstanceId)> Callback) const {}
+	virtual void ForEachSMInstance(int32 InstanceIndex, TFunctionRef<bool(FSMInstanceId)> Callback) const {}
+
 	virtual void BeginUpdate() {}
 	virtual void EndUpdate() {}
 	virtual void Refresh(bool Async, bool Force) {}
@@ -294,7 +300,9 @@ struct FFoliageInfo
 	FOLIAGE_API UHierarchicalInstancedStaticMeshComponent* GetComponent() const;
 
 	void AddReferencedObjects(UObject* InThis, FReferenceCollector& Collector);
-		
+	void PostSerialize(FArchive& Ar);
+	void PostLoad();
+
 	FOLIAGE_API void CreateImplementation(EFoliageImplType InType);
 	FOLIAGE_API void Initialize(const UFoliageType* FoliageType);
 	FOLIAGE_API void Uninitialize();
@@ -356,10 +364,13 @@ struct FFoliageInfo
 
 	// For debugging. Validate state after editing.
 	void CheckValid();
-		
+	
+	void ForEachSMInstance(TFunctionRef<bool(FSMInstanceId)> Callback) const;
+	void ForEachSMInstance(int32 InstanceIndex, TFunctionRef<bool(FSMInstanceId)> Callback) const;
+
 	FOLIAGE_API void Refresh(bool Async, bool Force);
 	FOLIAGE_API void OnHiddenEditorViewMaskChanged(uint64 InHiddenEditorViews);
-	FOLIAGE_API void PostEditUndo(AInstancedFoliageActor* InIFA, UFoliageType* FoliageType);
+	FOLIAGE_API void PostEditUndo(UFoliageType* FoliageType);
 	FOLIAGE_API void PreEditUndo(UFoliageType* FoliageType);
 	FOLIAGE_API void EnterEditMode();
 	FOLIAGE_API void ExitEditMode();
@@ -386,6 +397,49 @@ private:
 	void RemoveInstancesImpl(TArrayView<const int32> InInstancesToRemove, bool RebuildFoliageTree, FRemoveImplementationFunc ImplementationFunc);
 };
 
+struct FFoliageInstanceId
+{
+	explicit operator bool() const
+	{
+		return Info != nullptr
+			&& Index != INDEX_NONE;
+	}
+
+	bool operator==(const FFoliageInstanceId& InRHS) const
+	{
+		return Info == InRHS.Info
+			&& Index == InRHS.Index;
+	}
+
+	bool operator!=(const FFoliageInstanceId& InRHS) const
+	{
+		return !(*this == InRHS);
+	}
+
+#if WITH_EDITOR
+	FFoliageInstance* GetInstance() const
+	{
+		return (Info && Info->Instances.IsValidIndex(Index))
+			? &Info->Instances[Index]
+			: nullptr;
+	}
+
+	FFoliageInstance& GetInstanceChecked() const
+	{
+		FFoliageInstance* Instance = GetInstance();
+		check(Instance);
+		return *Instance;
+	}
+#endif
+
+	friend inline uint32 GetTypeHash(const FFoliageInstanceId& InId)
+	{
+		return HashCombine(GetTypeHash(InId.Index), GetTypeHash(InId.Info));
+	}
+
+	FFoliageInfo* Info = nullptr;
+	int32 Index = INDEX_NONE;
+};
 
 #if WITH_EDITORONLY_DATA
 //
diff --git a/Engine/Source/Runtime/Foliage/Public/InstancedFoliageActor.h b/Engine/Source/Runtime/Foliage/Public/InstancedFoliageActor.h
--- a/Engine/Source/Runtime/Foliage/Public/InstancedFoliageActor.h
+++ b/Engine/Source/Runtime/Foliage/Public/InstancedFoliageActor.h
@@ -21,7 +21,7 @@ class UProceduralFoliageComponent;
 typedef TFunction<bool(const UPrimitiveComponent*)> FFoliageTraceFilterFunc;
 
 UCLASS(notplaceable, hidecategories = (Object, Rendering, Mobility), MinimalAPI, NotBlueprintable)
-class AInstancedFoliageActor : public AISMPartitionActor, public ISMInstanceManager
+class AInstancedFoliageActor : public AISMPartitionActor
 {
 	GENERATED_UCLASS_BODY()
 
@@ -56,30 +56,31 @@ public:
 	virtual void RerunConstructionScripts() override {}
 	virtual bool IsLevelBoundsRelevant() const override { return false; }
 
-protected:
-	//~ ISMInstanceManagerProvider interface
-	virtual ISMInstanceManager* GetSMInstanceManager(const FSMInstanceId& InstanceId) override;
-
-	//~ ISMInstanceManager interface
-	virtual bool CanEditSMInstance(const FSMInstanceId& InstanceId) const override;
-	virtual bool CanMoveSMInstance(const FSMInstanceId& InstanceId, const ETypedElementWorldType InWorldType) const override;
-	virtual bool GetSMInstanceTransform(const FSMInstanceId& InstanceId, FTransform& OutInstanceTransform, bool bWorldSpace = false) const override;
-	virtual bool SetSMInstanceTransform(const FSMInstanceId& InstanceId, const FTransform& InstanceTransform, bool bWorldSpace = false, bool bMarkRenderStateDirty = false, bool bTeleport = false) override;
-	virtual void NotifySMInstanceMovementStarted(const FSMInstanceId& InstanceId) override;
-	virtual void NotifySMInstanceMovementOngoing(const FSMInstanceId& InstanceId) override;
-	virtual void NotifySMInstanceMovementEnded(const FSMInstanceId& InstanceId) override;
-	virtual void NotifySMInstanceSelectionChanged(const FSMInstanceId& InstanceId, const bool bIsSelected) override;
-	virtual bool DeleteSMInstances(TArrayView<const FSMInstanceId> InstanceIds) override;
-	virtual bool DuplicateSMInstances(TArrayView<const FSMInstanceId> InstanceIds, TArray<FSMInstanceId>& OutNewInstanceIds) override;
-
+public:
 #if WITH_EDITOR
-	void HandleSMInstancePreMove(const FSMInstanceId& InstanceId, FFoliageInfo& FoliageInfo);
-	void HandleSMInstancePostMove(const FSMInstanceId& InstanceId, FFoliageInfo& FoliageInfo, UFoliageType* FoliageType);
+	bool CanEditFoliageInstance(const FFoliageInstanceId& InstanceId) const;
+	bool CanMoveFoliageInstance(const FFoliageInstanceId& InstanceId, const ETypedElementWorldType WorldType) const;
+	bool GetFoliageInstanceTransform(const FFoliageInstanceId& InstanceId, FTransform& OutInstanceTransform, bool bWorldSpace) const;
+	bool SetFoliageInstanceTransform(const FFoliageInstanceId& InstanceId, const FTransform& InstanceTransform, bool bWorldSpace, bool bTeleport);
+	void NotifyFoliageInstanceMovementStarted(const FFoliageInstanceId& InstanceId);
+	void NotifyFoliageInstanceMovementOngoing(const FFoliageInstanceId& InstanceId);
+	void NotifyFoliageInstanceMovementEnded(const FFoliageInstanceId& InstanceId);
+	void NotifyFoliageInstanceSelectionChanged(const FFoliageInstanceId& InstanceId, const bool bIsSelected);
+	bool DeleteFoliageInstances(TArrayView<const FFoliageInstanceId> InstanceIds);
+	bool DuplicateFoliageInstances(TArrayView<const FFoliageInstanceId> InstanceIds, TArray<FFoliageInstanceId>& OutNewInstanceIds);
+
+	UFoliageType* GetFoliageTypeForInfo(const FFoliageInfo* FoliageInfo) const;
+#endif
 
-	FFoliageInfo* GetFoliageInfoForSMInstance(const FSMInstanceId& InstanceId, UFoliageType** OutFoliageType = nullptr);
-	const FFoliageInfo* GetFoliageInfoForSMInstance(const FSMInstanceId& InstanceId, UFoliageType** OutFoliageType = nullptr) const;
+protected:
+#if WITH_EDITOR
+	void HandleFoliageInstancePreMove(const FFoliageInstanceId& InstanceId);
+	void HandleFoliageInstancePostMove(const FFoliageInstanceId& InstanceId);
 #endif
 
+	//~ ISMInstanceManagerProvider interface
+	virtual ISMInstanceManager* GetSMInstanceManager(const FSMInstanceId& InstanceId) override;
+
 	// Default InternalTakeRadialDamage behavior finds and scales damage for the closest component which isn't appropriate for foliage.
 	virtual float InternalTakeRadialDamage(float Damage, struct FRadialDamageEvent const& RadialDamageEvent, class AController* EventInstigator, AActor* DamageCauser) override;
 
