Foliage: FFoliageInstanceHash LWC Support (fixes static analysis)

diff --git a/Engine/Source/Runtime/Foliage/Public/InstancedFoliage.h b/Engine/Source/Runtime/Foliage/Public/InstancedFoliage.h
--- a/Engine/Source/Runtime/Foliage/Public/InstancedFoliage.h
+++ b/Engine/Source/Runtime/Foliage/Public/InstancedFoliage.h
@@ -453,13 +453,13 @@ struct FFoliageInstanceHash
 private:
 	struct FKey
 	{
-		int32 X;
-		int32 Y;
-		int32 Z;
+		int64 X;
+		int64 Y;
+		int64 Z;
 
 		FKey() = default;
 
-		FKey(int32 InX, int32 InY, int32 InZ)
+		FKey(int64 InX, int64 InY, int64 InZ)
 			: X(InX), Y(InY), Z(InZ) {}
 
 		bool operator==(const FKey& Other) const
@@ -488,7 +488,7 @@ private:
 		
 	FKey MakeKey(const FVector& Location) const
 	{
-		return FKey(FMath::FloorToInt(Location.X) >> HashCellBits, FMath::FloorToInt(Location.Y) >> HashCellBits, FMath::FloorToInt(Location.Z) >> HashCellBits);
+		return FKey(FMath::FloorToInt64(Location.X) >> HashCellBits, FMath::FloorToInt64(Location.Y) >> HashCellBits, FMath::FloorToInt64(Location.Z) >> HashCellBits);
 	}
 
 	FVector MakeLocation(FKey CellKey) const
