Add checks to AInstancedFoliageActor::AddMesh in order to validate the UFoliageType in a partionned world

diff --git a/Engine/Source/Runtime/Foliage/Private/InstancedFoliage.cpp b/Engine/Source/Runtime/Foliage/Private/InstancedFoliage.cpp
--- a/Engine/Source/Runtime/Foliage/Private/InstancedFoliage.cpp
+++ b/Engine/Source/Runtime/Foliage/Private/InstancedFoliage.cpp
@@ -3609,6 +3609,13 @@ UFoliageType* AInstancedFoliageActor::AddFoliageType(const UFoliageType* InType,
 
 FFoliageInfo* AInstancedFoliageActor::AddMesh(UStaticMesh* InMesh, UFoliageType** OutSettings, const UFoliageType_InstancedStaticMesh* DefaultSettings)
 {
+	// This function is deprecated in a partioned world.
+	// FoliageType cannot have an AInstancedFoliageActor as their Outer
+	// This creates issues with the Foliage Edit mode.
+	// Proper way is to Create an asset for the UFoliageType and call
+	// the InstancedFoliageActor::AddMesh(UFoliageType*) overload with the asset
+	check(!GetWorld()->IsPartitionedWorld());
+
 	check(GetLocalFoliageTypeForSource(InMesh) == nullptr);
 
 	MarkPackageDirty();
@@ -3640,6 +3647,8 @@ FFoliageInfo* AInstancedFoliageActor::AddMesh(UStaticMesh* InMesh, UFoliageType*
 
 FFoliageInfo* AInstancedFoliageActor::AddMesh(UFoliageType* InType)
 {
+	check(!GetWorld()->IsPartitionedWorld() || InType->IsAsset());
+
 	check(FoliageInfos.Find(InType) == nullptr);
 
 	Modify();
