Foliage: Expose FoliageOverlappingBoxTransforms

diff --git a/Engine/Source/Runtime/Foliage/Public/FoliageStatistics.h b/Engine/Source/Runtime/Foliage/Public/FoliageStatistics.h
--- a/Engine/Source/Runtime/Foliage/Public/FoliageStatistics.h
+++ b/Engine/Source/Runtime/Foliage/Public/FoliageStatistics.h
@@ -40,6 +40,7 @@ class FOLIAGE_API UFoliageStatistics : public UBlueprintFunctionLibrary
 	 *	@param Box Box to use for overlap
 	 *	@param OutTransforms Array to populate with transforms
 	 */
+	UFUNCTION(BlueprintCallable, Category = "Foliage", meta = (WorldContext = "WorldContextObject", UnsafeDuringActorConstruction = "true"))
 	static void FoliageOverlappingBoxTransforms(UObject* WorldContextObject, const UStaticMesh* StaticMesh, FBox Box, TArray<FTransform>& OutTransforms);
 
 	/** 
