Make IsAssetOrNotBlueprint once again always available
Fix checks against whether something is not a blueprint to use IsNative which works in cooked builds
#rnx

diff --git a/Engine/Source/Runtime/Foliage/Private/InstancedFoliage.cpp b/Engine/Source/Runtime/Foliage/Private/InstancedFoliage.cpp
--- a/Engine/Source/Runtime/Foliage/Private/InstancedFoliage.cpp
+++ b/Engine/Source/Runtime/Foliage/Private/InstancedFoliage.cpp
@@ -715,12 +715,10 @@ void UFoliageType::PostLoad()
 	}
 }
 
-#if WITH_EDITOR
 bool UFoliageType::IsNotAssetOrBlueprint() const
 {
-	return IsAsset() == false && Cast<UBlueprint>(GetClass()->ClassGeneratedBy) == nullptr;
+	return IsAsset() == false && GetClass()->IsNative();
 }
-#endif
 
 FVector UFoliageType::GetRandomScale() const
 {
@@ -2951,12 +2949,7 @@ UFoliageType* AInstancedFoliageActor::GetLocalFoliageTypeForSource(const UObject
 		UFoliageType* FoliageType = Pair.Key;
 		// Check that the type is neither an asset nor blueprint instance
 
-		// ClassGeneratedBy TODO: This is wrong in cooked builds
-		if (FoliageType && FoliageType->GetSource() == InSource && !FoliageType->IsAsset()
-#if WITH_EDITORONLY_DATA
-		 && !FoliageType->GetClass()->ClassGeneratedBy
-#endif
-		)
+		if (FoliageType && FoliageType->GetSource() == InSource && !FoliageType->IsAsset() && FoliageType->GetClass()->IsNative())
 		{
 			ReturnType = FoliageType;
 			Info = &*Pair.Value;
