-Call RefreshActivePaletteViewWidget immediately in UpdatePalette, this ony request a new layouting for the tile.
-This is necessary if we're refreshing because a UFoliageType asset is deleted, since FoliagePaletteItems are mostly invalid when the FoliageType they're using gets deleted and can't participate in layouting until all refreshes are done. (this avoids adding if (!nullptr) everywhere)
-Added call to FoliageInfo->Uninitialize() in AInstancedFoliageActor::PostLoad to make sure stale InstancedMeshComponents don't get stuck in the AInstancedFoliageActor if their UFoliageType type assets get deleted while not loaded.

diff --git a/Engine/Source/Editor/FoliageEdit/Private/SFoliagePalette.cpp b/Engine/Source/Editor/FoliageEdit/Private/SFoliagePalette.cpp
--- a/Engine/Source/Editor/FoliageEdit/Private/SFoliagePalette.cpp
+++ b/Engine/Source/Editor/FoliageEdit/Private/SFoliagePalette.cpp
@@ -401,6 +401,9 @@ void SFoliagePalette::UpdatePalette(bool bRebuildItems)
 	{
 		bIsRebuildTimerRegistered = true;
 		RegisterActiveTimer(0.f, FWidgetActiveTimerDelegate::CreateSP(this, &SFoliagePalette::UpdatePaletteItems));
+
+		// widget view needs to be refreshed immediately since it's possibly invalid and can't participate in layouting
+		RefreshActivePaletteViewWidget();
 	}
 }
 
diff --git a/Engine/Source/Runtime/Foliage/Private/InstancedFoliage.cpp b/Engine/Source/Runtime/Foliage/Private/InstancedFoliage.cpp
--- a/Engine/Source/Runtime/Foliage/Private/InstancedFoliage.cpp
+++ b/Engine/Source/Runtime/Foliage/Private/InstancedFoliage.cpp
@@ -4346,17 +4346,18 @@ void AInstancedFoliageActor::PostLoad()
 		}
 			   
 		{
-			bool bContainsNull = FoliageInfos.Remove(nullptr) > 0;
-			if (bContainsNull)
+			if (FoliageInfos.Contains(nullptr))
 			{
 				FMessageLog("MapCheck").Warning()
 					->AddToken(FUObjectToken::Create(this))
 					->AddToken(FTextToken::Create(LOCTEXT("MapCheck_Message_FoliageMissingStaticMesh", "Foliage instances for a missing static mesh have been removed.")))
 					->AddToken(FMapErrorToken::Create(FMapErrors::FoliageMissingStaticMesh));
-				while (bContainsNull)
+								
+				TUniqueObj<FFoliageInfo> FoliageInfo;
+				while (FoliageInfos.RemoveAndCopyValue(nullptr, FoliageInfo))
 				{
-					bContainsNull = FoliageInfos.Remove(nullptr) > 0;
-				}
+					FoliageInfo->Uninitialize();
+				}				
 			}
 		}
 
