Fixed crash when doing foliage info fixup in AInstancedFoliageActor loading.

diff --git a/Engine/Source/Runtime/Foliage/Private/InstancedFoliage.cpp b/Engine/Source/Runtime/Foliage/Private/InstancedFoliage.cpp
--- a/Engine/Source/Runtime/Foliage/Private/InstancedFoliage.cpp
+++ b/Engine/Source/Runtime/Foliage/Private/InstancedFoliage.cpp
@@ -4356,7 +4356,10 @@ void AInstancedFoliageActor::PostLoad()
 				TUniqueObj<FFoliageInfo> FoliageInfo;
 				while (FoliageInfos.RemoveAndCopyValue(nullptr, FoliageInfo))
 				{
-					FoliageInfo->Uninitialize();
+					if (FoliageInfo->IsInitialized())
+					{
+						FoliageInfo->Uninitialize();
+					}
 				}				
 			}
 		}
