Add some logic for bEvaluateWorldPositionInRayTracing that was missed in original conversion from bEvaluateWorldPosition.

diff --git a/Engine/Source/Runtime/Engine/Private/Components/StaticMeshComponent.cpp b/Engine/Source/Runtime/Engine/Private/Components/StaticMeshComponent.cpp
--- a/Engine/Source/Runtime/Engine/Private/Components/StaticMeshComponent.cpp
+++ b/Engine/Source/Runtime/Engine/Private/Components/StaticMeshComponent.cpp
@@ -202,6 +202,7 @@ UStaticMeshComponent::UStaticMeshComponent(const FObjectInitializer& ObjectIniti
 	bDisallowMeshPaintPerInstance = false;
 	bDisallowNanite = false;
 	bEvaluateWorldPositionOffset = false;
+	bEvaluateWorldPositionOffsetInRayTracing = false;
 	DistanceFieldIndirectShadowMinVisibility = .1f;
 	GetBodyInstance()->bAutoWeld = true;	//static mesh by default has auto welding
 
diff --git a/Engine/Source/Runtime/Engine/Private/PreviewScene.cpp b/Engine/Source/Runtime/Engine/Private/PreviewScene.cpp
--- a/Engine/Source/Runtime/Engine/Private/PreviewScene.cpp
+++ b/Engine/Source/Runtime/Engine/Private/PreviewScene.cpp
@@ -161,6 +161,7 @@
 		if(pStaticMesh != nullptr)
 		{
 			pStaticMesh->bEvaluateWorldPositionOffset = true;
+			pStaticMesh->bEvaluateWorldPositionOffsetInRayTracing = true;
 		}
 	}
 
diff --git a/Engine/Source/Runtime/Foliage/Private/InstancedFoliage.cpp b/Engine/Source/Runtime/Foliage/Private/InstancedFoliage.cpp
--- a/Engine/Source/Runtime/Foliage/Private/InstancedFoliage.cpp
+++ b/Engine/Source/Runtime/Foliage/Private/InstancedFoliage.cpp
@@ -1626,6 +1626,11 @@
 			Component->bEvaluateWorldPositionOffset = FoliageType->bEvaluateWorldPositionOffset;
 			bNeedsMarkRenderStateDirty = true;
 		}
+		if (Component->bEvaluateWorldPositionOffsetInRayTracing != FoliageType->bEvaluateWorldPositionOffset)
+		{
+			Component->bEvaluateWorldPositionOffsetInRayTracing = FoliageType->bEvaluateWorldPositionOffset;
+			bNeedsMarkRenderStateDirty = true;
+		}
 
 
 		if (Component->bEnableDensityScaling != FoliageType->bEnableDensityScaling)
