Foliage: show foliage actors in outliner

diff --git a/Engine/Source/Runtime/Foliage/Private/InstancedFoliage.cpp b/Engine/Source/Runtime/Foliage/Private/InstancedFoliage.cpp
--- a/Engine/Source/Runtime/Foliage/Private/InstancedFoliage.cpp
+++ b/Engine/Source/Runtime/Foliage/Private/InstancedFoliage.cpp
@@ -4303,6 +4303,13 @@ void AInstancedFoliageActor::BeginDestroy()
 		FoliageInfos.Empty();
 	}
 }
+
+bool AInstancedFoliageActor::IsListedInSceneOutliner() const
+{
+	// In World Partition the AInstancedFoliageActor::GetDefault creates a transient foliage actor to store unpainted foliage types.
+	return !HasAnyFlags(RF_Transient);
+}
+
 #endif
 
 void AInstancedFoliageActor::PostLoad()
diff --git a/Engine/Source/Runtime/Foliage/Public/InstancedFoliageActor.h b/Engine/Source/Runtime/Foliage/Public/InstancedFoliageActor.h
--- a/Engine/Source/Runtime/Foliage/Public/InstancedFoliageActor.h
+++ b/Engine/Source/Runtime/Foliage/Public/InstancedFoliageActor.h
@@ -98,6 +98,8 @@
 	virtual void PostInitProperties() override;
 	virtual void BeginDestroy() override;
 	virtual void Destroyed() override;
+	virtual bool IsListedInSceneOutliner() const override;
+	
 	FOLIAGE_API void CleanupDeletedFoliageType();
 	FOLIAGE_API void DetectFoliageTypeChangeAndUpdate();
 
