ProceduralFoliageComponent: Fix crash in non World Partition world (details panel)

#rnx

diff --git a/Engine/Source/Runtime/Foliage/Private/ProceduralFoliageComponent.cpp b/Engine/Source/Runtime/Foliage/Private/ProceduralFoliageComponent.cpp
--- a/Engine/Source/Runtime/Foliage/Private/ProceduralFoliageComponent.cpp
+++ b/Engine/Source/Runtime/Foliage/Private/ProceduralFoliageComponent.cpp
@@ -106,8 +106,10 @@ void UProceduralFoliageComponent::LoadSimulatedRegion()
 {
 	if (GetOwner()->Implements<UWorldPartitionActorLoaderInterface>())
 	{
-		IWorldPartitionActorLoaderInterface::ILoaderAdapter* LoaderAdapter = Cast<IWorldPartitionActorLoaderInterface>(GetOwner())->GetLoaderAdapter();
-		LoaderAdapter->Load();
+		if (IWorldPartitionActorLoaderInterface::ILoaderAdapter* LoaderAdapter = Cast<IWorldPartitionActorLoaderInterface>(GetOwner())->GetLoaderAdapter())
+		{
+			LoaderAdapter->Load();
+		}
 	}
 }
 
@@ -115,8 +117,10 @@ void UProceduralFoliageComponent::UnloadSimulatedRegion()
 {
 	if (GetOwner()->Implements<UWorldPartitionActorLoaderInterface>())
 	{
-		IWorldPartitionActorLoaderInterface::ILoaderAdapter* LoaderAdapter = Cast<IWorldPartitionActorLoaderInterface>(GetOwner())->GetLoaderAdapter();
-		LoaderAdapter->Unload();
+		if (IWorldPartitionActorLoaderInterface::ILoaderAdapter* LoaderAdapter = Cast<IWorldPartitionActorLoaderInterface>(GetOwner())->GetLoaderAdapter())
+		{
+			LoaderAdapter->Unload();
+		}
 	}
 }
 
@@ -124,8 +128,10 @@ bool UProceduralFoliageComponent::IsSimulatedRegionLoaded()
 {
 	if (GetOwner()->Implements<UWorldPartitionActorLoaderInterface>())
 	{
-		IWorldPartitionActorLoaderInterface::ILoaderAdapter* LoaderAdapter = Cast<IWorldPartitionActorLoaderInterface>(GetOwner())->GetLoaderAdapter();
-		return LoaderAdapter->IsLoaded();
+		if (IWorldPartitionActorLoaderInterface::ILoaderAdapter* LoaderAdapter = Cast<IWorldPartitionActorLoaderInterface>(GetOwner())->GetLoaderAdapter())
+		{
+			return LoaderAdapter->IsLoaded();
+		}
 	}
 	return true;
 }
