ProceduralFoliageComponent: Allow passing in generation params

#rnx

diff --git a/Engine/Source/Runtime/Foliage/Private/InstancedFoliage.cpp b/Engine/Source/Runtime/Foliage/Private/InstancedFoliage.cpp
--- a/Engine/Source/Runtime/Foliage/Private/InstancedFoliage.cpp
+++ b/Engine/Source/Runtime/Foliage/Private/InstancedFoliage.cpp
@@ -3260,9 +3260,14 @@ void AInstancedFoliageActor::DeleteInstancesForComponent(UWorld* InWorld, UActor
 	}
 }
 
-void AInstancedFoliageActor::DeleteInstancesForProceduralFoliageComponent(const UProceduralFoliageComponent* ProceduralFoliageComponent, bool InRebuildTree)
+bool AInstancedFoliageActor::DeleteInstancesForProceduralFoliageComponent(const UProceduralFoliageComponent* ProceduralFoliageComponent, bool InRebuildTree)
 {
-	const FGuid& ProceduralGuid = ProceduralFoliageComponent->GetProceduralGuid();
+	return DeleteInstancesForProceduralFoliageComponent(ProceduralFoliageComponent->GetProceduralGuid(), InRebuildTree);
+}
+
+bool AInstancedFoliageActor::DeleteInstancesForProceduralFoliageComponent(const FGuid& ProceduralGuid, bool InRebuildTree)
+{
+	bool bRemovedInstances = false;
 	BeginUpdate();
 	for (auto& Pair : FoliageInfos)
 	{
@@ -3279,16 +3284,23 @@ void AInstancedFoliageActor::DeleteInstancesForProceduralFoliageComponent(const
 		if (InstancesToRemove.Num())
 		{
 			Info.RemoveInstances(InstancesToRemove, InRebuildTree);
+			bRemovedInstances = true;
 		}
 	}
 	EndUpdate();
 	// Clean up dead cross-level references
 	FFoliageInstanceBaseCache::CompactInstanceBaseCache(this);
+
+	return bRemovedInstances;
 }
 
 bool AInstancedFoliageActor::ContainsInstancesFromProceduralFoliageComponent(const UProceduralFoliageComponent* ProceduralFoliageComponent)
 {
-	const FGuid& ProceduralGuid = ProceduralFoliageComponent->GetProceduralGuid();
+	return ContainsInstancesFromProceduralFoliageComponent(ProceduralFoliageComponent->GetProceduralGuid());
+}
+
+bool AInstancedFoliageActor::ContainsInstancesFromProceduralFoliageComponent(const FGuid& ProceduralGuid)
+{
 	for (auto& Pair : FoliageInfos)
 	{
 		FFoliageInfo& Info = *Pair.Value;
diff --git a/Engine/Source/Runtime/Foliage/Private/ProceduralFoliageComponent.cpp b/Engine/Source/Runtime/Foliage/Private/ProceduralFoliageComponent.cpp
--- a/Engine/Source/Runtime/Foliage/Private/ProceduralFoliageComponent.cpp
+++ b/Engine/Source/Runtime/Foliage/Private/ProceduralFoliageComponent.cpp
@@ -36,7 +36,7 @@ UProceduralFoliageComponent::UProceduralFoliageComponent(const FObjectInitialize
 #endif
 }
 
-FBox2D GetTileRegion(const int32 X, const int32 Y, const float InnerSize, const float Overlap)
+static FBox2D GetTileRegion(const int32 X, const int32 Y, const float InnerSize, const float Overlap)
 {
 	// For tiles on the bottom and the left (that have no preceding neighbor to overlap them), have the
 	// tile fill that space with its own content
@@ -79,28 +79,44 @@ FBodyInstance* UProceduralFoliageComponent::GetBoundsBodyInstance() const
 	return nullptr;
 }
 
-void UProceduralFoliageComponent::GetTileLayout(FTileLayout& OutTileLayout) const
+UProceduralFoliageComponent::FGenerateProceduralContentParams UProceduralFoliageComponent::GetGenerateProceduralContentParams() const
+{
+	FGenerateProceduralContentParams Params;
+	Params.Bounds = GetBounds();
+	Params.FoliageSpawner = FoliageSpawner;
+	Params.ProceduralGuid = GetProceduralGuid();
+	Params.TileOverlap = TileOverlap;
+	Params.ProceduralVolumeInstance = GetBoundsBodyInstance();
+
+	return Params;
+}
+
+void UProceduralFoliageComponent::GetTileLayout(const FGenerateProceduralContentParams& Params, FTileLayout& OutTileLayout)
 {
-	FBox Bounds = GetBounds();
-	if (Bounds.IsValid)
+	if (Params.Bounds.IsValid)
 	{
 		// Determine the bottom-left-most tile that contains the min position (when accounting for overlap)
-		const FVector MinPosition = Bounds.Min + TileOverlap;
-		OutTileLayout.BottomLeftX = FMath::FloorToInt(MinPosition.X / FoliageSpawner->TileSize);
-		OutTileLayout.BottomLeftY = FMath::FloorToInt(MinPosition.Y / FoliageSpawner->TileSize);
+		const FVector MinPosition = Params.Bounds.Min + Params.TileOverlap;
+		OutTileLayout.BottomLeftX = FMath::FloorToInt(MinPosition.X / Params.FoliageSpawner->TileSize);
+		OutTileLayout.BottomLeftY = FMath::FloorToInt(MinPosition.Y / Params.FoliageSpawner->TileSize);
 
 		// Determine the total number of tiles along each active axis
-		const FVector MaxPosition = Bounds.Max - TileOverlap;
-		const int32 MaxXIdx = FMath::FloorToInt(MaxPosition.X / FoliageSpawner->TileSize);
-		const int32 MaxYIdx = FMath::FloorToInt(MaxPosition.Y / FoliageSpawner->TileSize);
+		const FVector MaxPosition = Params.Bounds.Max - Params.TileOverlap;
+		const int32 MaxXIdx = FMath::FloorToInt(MaxPosition.X / Params.FoliageSpawner->TileSize);
+		const int32 MaxYIdx = FMath::FloorToInt(MaxPosition.Y / Params.FoliageSpawner->TileSize);
 
 		OutTileLayout.NumTilesX = (MaxXIdx - OutTileLayout.BottomLeftX) + 1;
 		OutTileLayout.NumTilesY = (MaxYIdx - OutTileLayout.BottomLeftY) + 1;
 
-		OutTileLayout.HalfHeight = Bounds.GetExtent().Z;
+		OutTileLayout.HalfHeight = Params.Bounds.GetExtent().Z;
 	}
 }
 
+void UProceduralFoliageComponent::GetTileLayout(FTileLayout& OutTileLayout) const
+{
+	GetTileLayout(GetGenerateProceduralContentParams(), OutTileLayout);
+}
+
 #if WITH_EDITOR
 void UProceduralFoliageComponent::LoadSimulatedRegion()
 {
@@ -139,26 +155,37 @@ bool UProceduralFoliageComponent::IsSimulatedRegionLoaded()
 
 FVector UProceduralFoliageComponent::GetWorldPosition() const
 {
-	FBox Bounds = GetBounds();
-	if (!Bounds.IsValid || FoliageSpawner == nullptr)
+	return GetWorldPosition(GetGenerateProceduralContentParams());
+}
+
+FVector UProceduralFoliageComponent::GetWorldPosition(const FGenerateProceduralContentParams& Params)
+{
+	if (!Params.Bounds.IsValid || Params.FoliageSpawner == nullptr)
 	{
 		return FVector::ZeroVector;
 	}
 	
 	FTileLayout TileLayout;
-	GetTileLayout(TileLayout);
+	GetTileLayout(Params, TileLayout);
 
-	const float TileSize = FoliageSpawner->TileSize;
-	return FVector(TileLayout.BottomLeftX * TileSize, TileLayout.BottomLeftY * TileSize, Bounds.GetCenter().Z);
+	const float TileSize = Params.FoliageSpawner->TileSize;
+	return FVector(TileLayout.BottomLeftX * TileSize, TileLayout.BottomLeftY * TileSize, Params.Bounds.GetCenter().Z);
 }
 
-bool UProceduralFoliageComponent::ExecuteSimulation(TArray<FDesiredFoliageInstance>& OutInstances)
+bool UProceduralFoliageComponent::GenerateProceduralContent(TArray<FDesiredFoliageInstance>& OutInstances)
 {
 #if WITH_EDITOR
 	LoadSimulatedRegion();
+	return GenerateProceduralContent(GetGenerateProceduralContentParams(), OutInstances);
+#endif
 
-	FBodyInstance* BoundsBodyInstance = GetBoundsBodyInstance();
-	if (FoliageSpawner)
+	return false;
+}
+
+bool UProceduralFoliageComponent::GenerateProceduralContent(const FGenerateProceduralContentParams& Params, TArray<FDesiredFoliageInstance>& OutInstances)
+{
+#if WITH_EDITOR
+	if (Params.FoliageSpawner)
 	{
 		// Establish the counter used to check if the user has canceled the simulation
 		FThreadSafeCounter LastCancel;
@@ -166,15 +193,15 @@ bool UProceduralFoliageComponent::ExecuteSimulation(TArray<FDesiredFoliageInstan
 		FThreadSafeCounter* LastCancelPtr = &LastCancel;
 		
 		// Establish basic info about the tiles
-		const float TileSize = FoliageSpawner->TileSize;
-		const FVector WorldPosition = GetWorldPosition();
-		FBox ActorVolumeBounds = GetBounds();
-		FVector2D ActorVolumeLocation = FVector2D(ActorVolumeBounds.GetCenter());
-		const float ActorVolumeMaxExtent = FVector2D(ActorVolumeBounds.GetExtent()).GetMax();
+		const float TileOverlap = Params.TileOverlap;
+		const float TileSize = Params.FoliageSpawner->TileSize;
+		const FVector WorldPosition = GetWorldPosition(Params);
+		FVector2D GenerateLocation = FVector2D(Params.Bounds.GetCenter());
+		const float GenerateMaxExtent = FVector2D(Params.Bounds.GetExtent()).GetMax();
 		FTileLayout TileLayout;
-		GetTileLayout(TileLayout);
+		GetTileLayout(Params, TileLayout);
 
-		FoliageSpawner->Simulate();
+		Params.FoliageSpawner->Simulate();
 
 		TArray<TFuture< TArray<FDesiredFoliageInstance>* >> Futures;
 		for (int32 X = 0; X < TileLayout.NumTilesX; ++X)
@@ -182,7 +209,7 @@ bool UProceduralFoliageComponent::ExecuteSimulation(TArray<FDesiredFoliageInstan
 			for (int32 Y = 0; Y < TileLayout.NumTilesY; ++Y)
 			{
 				// We have to get the simulated tiles and create new ones to build on main thread
-				const UProceduralFoliageTile* Tile = FoliageSpawner->GetRandomTile(X + TileLayout.BottomLeftX, Y + TileLayout.BottomLeftY);
+				const UProceduralFoliageTile* Tile = Params.FoliageSpawner->GetRandomTile(X + TileLayout.BottomLeftX, Y + TileLayout.BottomLeftY);
 				if (Tile == nullptr)	
 				{
 					// Simulation was either canceled or failed
@@ -190,12 +217,12 @@ bool UProceduralFoliageComponent::ExecuteSimulation(TArray<FDesiredFoliageInstan
 				}
 
 				// From the pool of simulated tiles, pick the neighbors of this tile
-				const UProceduralFoliageTile* RightTile = (X + 1 < TileLayout.NumTilesX) ? FoliageSpawner->GetRandomTile(X + TileLayout.BottomLeftX + 1, Y + TileLayout.BottomLeftY) : nullptr;
-				const UProceduralFoliageTile* TopTile = (Y + 1 < TileLayout.NumTilesY) ? FoliageSpawner->GetRandomTile(X + TileLayout.BottomLeftX, Y + TileLayout.BottomLeftY + 1) : nullptr;
-				const UProceduralFoliageTile* TopRightTile = (RightTile && TopTile) ? FoliageSpawner->GetRandomTile(X + TileLayout.BottomLeftX + 1, Y + TileLayout.BottomLeftY + 1) : nullptr;
+				const UProceduralFoliageTile* RightTile = (X + 1 < TileLayout.NumTilesX) ? Params.FoliageSpawner->GetRandomTile(X + TileLayout.BottomLeftX + 1, Y + TileLayout.BottomLeftY) : nullptr;
+				const UProceduralFoliageTile* TopTile = (Y + 1 < TileLayout.NumTilesY) ? Params.FoliageSpawner->GetRandomTile(X + TileLayout.BottomLeftX, Y + TileLayout.BottomLeftY + 1) : nullptr;
+				const UProceduralFoliageTile* TopRightTile = (RightTile && TopTile) ? Params.FoliageSpawner->GetRandomTile(X + TileLayout.BottomLeftX + 1, Y + TileLayout.BottomLeftY + 1) : nullptr;
 
 				// Create a temp tile that will contain the composite contents of the tile after accounting for overlap
-				UProceduralFoliageTile* CompositeTile = FoliageSpawner->CreateTempTile();
+				UProceduralFoliageTile* CompositeTile = Params.FoliageSpawner->CreateTempTile();
 
 				Futures.Add(Async(EAsyncExecution::ThreadPool, [=]()
 				{
@@ -208,7 +235,7 @@ bool UProceduralFoliageComponent::ExecuteSimulation(TArray<FDesiredFoliageInstan
 					//@todo proc foliage: Determine the composite contents of the tile (including overlaps) without copying everything to a temp tile
 
 					// Copy the base tile contents
-					const FBox2D BaseTile = GetTileRegion(X, Y, TileSize, TileOverlap);
+					const FBox2D BaseTile = GetTileRegion(X, Y, TileSize, Params.TileOverlap);
 					Tile->CopyInstancesToTile(CompositeTile, BaseTile, FTransform::Identity, TileOverlap);
 
 
@@ -240,7 +267,7 @@ bool UProceduralFoliageComponent::ExecuteSimulation(TArray<FDesiredFoliageInstan
 					const FTransform TileTM(OrientedOffset + WorldPosition);
 
 					TArray<FDesiredFoliageInstance>* DesiredInstances = new TArray<FDesiredFoliageInstance>();
-					CompositeTile->ExtractDesiredInstances(*DesiredInstances, TileTM, ActorVolumeLocation, ActorVolumeMaxExtent, ProceduralGuid, TileLayout.HalfHeight, BoundsBodyInstance, true);
+					CompositeTile->ExtractDesiredInstances(*DesiredInstances, TileTM, GenerateLocation, GenerateMaxExtent, Params.ProceduralGuid, TileLayout.HalfHeight, Params.ProceduralVolumeInstance, true);
 
 					return DesiredInstances;
 					
@@ -295,7 +322,7 @@ bool UProceduralFoliageComponent::ExecuteSimulation(TArray<FDesiredFoliageInstan
 			for (int Y = 0; Y < TileLayout.NumTilesY; ++Y)
 			{
 				TArray<FDesiredFoliageInstance>* DesiredInstances = Futures[FutureIdx++].Get();
-				OutInstances.Append(*DesiredInstances);
+				OutInstances.Append(MoveTemp(*DesiredInstances));
 				delete DesiredInstances;
 			}
 		}
@@ -344,21 +371,23 @@ bool UProceduralFoliageComponent::ResimulateProceduralFoliage(TFunctionRef<void(
 	return false;
 }
 
-bool UProceduralFoliageComponent::GenerateProceduralContent(TArray <FDesiredFoliageInstance>& OutInstances)
+void UProceduralFoliageComponent::RemoveProceduralContent(bool InRebuildTree)
 {
 #if WITH_EDITOR
-	return ExecuteSimulation(OutInstances);
+	TSet<AInstancedFoliageActor*> OutModifiedActors;
+	RemoveProceduralContent(GetWorld(), GetProceduralGuid(), InRebuildTree, OutModifiedActors);
 #endif
-
-	return false;
 }
 
-void UProceduralFoliageComponent::RemoveProceduralContent(bool InRebuildTree)
+void UProceduralFoliageComponent::RemoveProceduralContent(UWorld* World, const FGuid& ProceduralGuid, bool InRebuildTree, TSet<AInstancedFoliageActor*>& OutModifiedActors)
 {
 #if WITH_EDITOR
-	for (TActorIterator<AInstancedFoliageActor> It(GetWorld()); It; ++It)
+	for (TActorIterator<AInstancedFoliageActor> It(World); It; ++It)
 	{
-		(*It)->DeleteInstancesForProceduralFoliageComponent(this, InRebuildTree);
+		if ((*It)->DeleteInstancesForProceduralFoliageComponent(ProceduralGuid, InRebuildTree))
+		{
+			OutModifiedActors.Add(*It);
+		}
 	}
 #endif
 }
diff --git a/Engine/Source/Runtime/Foliage/Public/InstancedFoliageActor.h b/Engine/Source/Runtime/Foliage/Public/InstancedFoliageActor.h
--- a/Engine/Source/Runtime/Foliage/Public/InstancedFoliageActor.h
+++ b/Engine/Source/Runtime/Foliage/Public/InstancedFoliageActor.h
@@ -219,10 +219,12 @@ public:
 	static FOLIAGE_API void DeleteInstancesForComponent(UWorld* InWorld, UActorComponent* InComponent);
 
 	// Deletes the instances spawned by a procedural component
-	void DeleteInstancesForProceduralFoliageComponent(const UProceduralFoliageComponent* ProceduralFoliageComponent, bool InRebuildTree);
-
+	FOLIAGE_API bool DeleteInstancesForProceduralFoliageComponent(const UProceduralFoliageComponent* ProceduralFoliageComponent, bool InRebuildTree);
+	FOLIAGE_API bool DeleteInstancesForProceduralFoliageComponent(const FGuid& ProceduralGuid, bool InRebuildTree);
+	
 	/** @return True if any instances exist that were spawned by the given procedural component */
-	bool ContainsInstancesFromProceduralFoliageComponent(const UProceduralFoliageComponent* ProceduralFoliageComponent);
+	FOLIAGE_API bool ContainsInstancesFromProceduralFoliageComponent(const UProceduralFoliageComponent* ProceduralFoliageComponent);
+	FOLIAGE_API bool ContainsInstancesFromProceduralFoliageComponent(const FGuid& ProceduralGuid);
 
 	// Finds a mesh entry or adds it if it doesn't already exist
 	FOLIAGE_API FFoliageInfo* FindOrAddMesh(UFoliageType* InType);
diff --git a/Engine/Source/Runtime/Foliage/Public/ProceduralFoliageComponent.h b/Engine/Source/Runtime/Foliage/Public/ProceduralFoliageComponent.h
--- a/Engine/Source/Runtime/Foliage/Public/ProceduralFoliageComponent.h
+++ b/Engine/Source/Runtime/Foliage/Public/ProceduralFoliageComponent.h
@@ -79,9 +79,31 @@ class FOLIAGE_API UProceduralFoliageComponent : public UActorComponent
 
 #endif
 
+	struct FGenerateProceduralContentParams
+	{
+		FGenerateProceduralContentParams()
+			: FoliageSpawner(nullptr)
+			, Bounds(EForceInit::ForceInit)
+			, TileOverlap(0.f)
+			, ProceduralVolumeInstance(nullptr)
+		{
+		}
+
+		TObjectPtr<UProceduralFoliageSpawner> FoliageSpawner;
+		FBox Bounds;
+		float TileOverlap;
+		FGuid ProceduralGuid;
+		FBodyInstance* ProceduralVolumeInstance;
+	};
+				
 	// UObject interface
 	virtual void PostEditImport() override;
 
+	/**
+	 * Returns a params struct based on this component properties. 
+	 */
+	FGenerateProceduralContentParams GetGenerateProceduralContentParams() const;
+
 	/**
 	 * Runs the procedural foliage simulation, removes the old result, creates instances with the new result
 	 * @return True if the simulation succeeded
@@ -92,16 +114,19 @@ class FOLIAGE_API UProceduralFoliageComponent : public UActorComponent
 	 * Runs the procedural foliage simulation to generate a list of desired instances to spawn.
 	 * @return True if the simulation succeeded
 	 */
-	bool GenerateProceduralContent(TArray<FDesiredFoliageInstance>& OutFoliageInstances);
-	
+	bool GenerateProceduralContent(TArray<FDesiredFoliageInstance>& OutInstances);
+	static bool GenerateProceduralContent(const FGenerateProceduralContentParams& Params, TArray<FDesiredFoliageInstance>& OutInstances);
+			
 	/** Removes all spawned foliage instances in the level that were spawned by this component */
 	void RemoveProceduralContent(bool InRebuildTree = true);
+	static void RemoveProceduralContent(UWorld* World, const FGuid& ProceduralGuid, bool InRebuildTree, TSet<AInstancedFoliageActor*>& OutModifiedActors);
 
 	/** @return True if any foliage instances in the level were spawned by this component */
 	bool HasSpawnedAnyInstances();
 	
 	/** @return The position in world space of the bottom-left corner of the bottom-left-most active tile */
 	FVector GetWorldPosition() const;
+	static FVector GetWorldPosition(const FGenerateProceduralContentParams& Param);
 
 	/** @return The bounds of area encompassed by the simulation */
 	virtual FBox GetBounds() const;
@@ -111,6 +136,7 @@ class FOLIAGE_API UProceduralFoliageComponent : public UActorComponent
 
 	/** Determines the basic layout of the tiles used in the simulation */
 	void GetTileLayout(FTileLayout& OutTileLayout) const;
+	static void GetTileLayout(const FGenerateProceduralContentParams& Params, FTileLayout& OutTileLayout);
 
 	void SetSpawningVolume(AVolume* InSpawningVolume) { SpawningVolume = InSpawningVolume; }
 	const FGuid& GetProceduralGuid() const { return ProceduralGuid; }
@@ -121,10 +147,6 @@ class FOLIAGE_API UProceduralFoliageComponent : public UActorComponent
 	bool IsSimulatedRegionLoaded();
 #endif
 
-private:
-	/** Does all the actual work of executing the procedural foliage simulation */
-	bool ExecuteSimulation(TArray<FDesiredFoliageInstance>& OutFoliageInstances);
-
 private:
 	UPROPERTY()
 	TObjectPtr<AVolume> SpawningVolume;
