Foliage: DeleteInstancesForAllProceduralFoliageComponents (Deletes all instances with valid ProceduralGuid)

diff --git a/Engine/Source/Runtime/Foliage/Private/InstancedFoliage.cpp b/Engine/Source/Runtime/Foliage/Private/InstancedFoliage.cpp
--- a/Engine/Source/Runtime/Foliage/Private/InstancedFoliage.cpp
+++ b/Engine/Source/Runtime/Foliage/Private/InstancedFoliage.cpp
@@ -3260,12 +3260,7 @@ void AInstancedFoliageActor::DeleteInstancesForComponent(UWorld* InWorld, UActor
 	}
 }
 
-bool AInstancedFoliageActor::DeleteInstancesForProceduralFoliageComponent(const UProceduralFoliageComponent* ProceduralFoliageComponent, bool InRebuildTree)
-{
-	return DeleteInstancesForProceduralFoliageComponent(ProceduralFoliageComponent->GetProceduralGuid(), InRebuildTree);
-}
-
-bool AInstancedFoliageActor::DeleteInstancesForProceduralFoliageComponent(const FGuid& ProceduralGuid, bool InRebuildTree)
+bool AInstancedFoliageActor::DeleteInstancesForProceduralFoliageComponentInternal(const FGuid& InProceduralGuid, bool bInRebuildTree, bool bInDeleteAll)
 {
 	bool bRemovedInstances = false;
 	BeginUpdate();
@@ -3275,7 +3270,7 @@ bool AInstancedFoliageActor::DeleteInstancesForProceduralFoliageComponent(const
 		TArray<int32> InstancesToRemove;
 		for (int32 InstanceIdx = 0; InstanceIdx < Info.Instances.Num(); InstanceIdx++)
 		{
-			if (Info.Instances[InstanceIdx].ProceduralGuid == ProceduralGuid)
+			if ((Info.Instances[InstanceIdx].ProceduralGuid == InProceduralGuid) || (bInDeleteAll && Info.Instances[InstanceIdx].ProceduralGuid.IsValid()))
 			{
 				InstancesToRemove.Add(InstanceIdx);
 			}
@@ -3283,7 +3278,7 @@ bool AInstancedFoliageActor::DeleteInstancesForProceduralFoliageComponent(const
 
 		if (InstancesToRemove.Num())
 		{
-			Info.RemoveInstances(InstancesToRemove, InRebuildTree);
+			Info.RemoveInstances(InstancesToRemove, bInRebuildTree);
 			bRemovedInstances = true;
 		}
 	}
@@ -3294,12 +3289,27 @@ bool AInstancedFoliageActor::DeleteInstancesForProceduralFoliageComponent(const
 	return bRemovedInstances;
 }
 
-bool AInstancedFoliageActor::ContainsInstancesFromProceduralFoliageComponent(const UProceduralFoliageComponent* ProceduralFoliageComponent)
+bool AInstancedFoliageActor::DeleteInstancesForAllProceduralFoliageComponents(bool bInRebuildTree)
+{
+	return DeleteInstancesForProceduralFoliageComponentInternal(FGuid(), bInRebuildTree, true);
+}
+
+bool AInstancedFoliageActor::DeleteInstancesForProceduralFoliageComponent(const UProceduralFoliageComponent* InProceduralFoliageComponent, bool bInRebuildTree)
+{
+	return DeleteInstancesForProceduralFoliageComponent(InProceduralFoliageComponent->GetProceduralGuid(), bInRebuildTree);
+}
+
+bool AInstancedFoliageActor::DeleteInstancesForProceduralFoliageComponent(const FGuid& InProceduralGuid, bool bInRebuildTree)
+{
+	return DeleteInstancesForProceduralFoliageComponentInternal(InProceduralGuid, bInRebuildTree, false);
+}
+
+bool AInstancedFoliageActor::ContainsInstancesFromProceduralFoliageComponent(const UProceduralFoliageComponent* InProceduralFoliageComponent)
 {
-	return ContainsInstancesFromProceduralFoliageComponent(ProceduralFoliageComponent->GetProceduralGuid());
+	return ContainsInstancesFromProceduralFoliageComponent(InProceduralFoliageComponent->GetProceduralGuid());
 }
 
-bool AInstancedFoliageActor::ContainsInstancesFromProceduralFoliageComponent(const FGuid& ProceduralGuid)
+bool AInstancedFoliageActor::ContainsInstancesFromProceduralFoliageComponent(const FGuid& InProceduralGuid)
 {
 	for (auto& Pair : FoliageInfos)
 	{
@@ -3307,7 +3317,7 @@ bool AInstancedFoliageActor::ContainsInstancesFromProceduralFoliageComponent(con
 		TArray<int32> InstancesToRemove;
 		for (int32 InstanceIdx = 0; InstanceIdx < Info.Instances.Num(); InstanceIdx++)
 		{
-			if (Info.Instances[InstanceIdx].ProceduralGuid == ProceduralGuid)
+			if (Info.Instances[InstanceIdx].ProceduralGuid == InProceduralGuid)
 			{
 				// The procedural component is responsible for an instance
 				return true;
diff --git a/Engine/Source/Runtime/Foliage/Private/ProceduralFoliageComponent.cpp b/Engine/Source/Runtime/Foliage/Private/ProceduralFoliageComponent.cpp
--- a/Engine/Source/Runtime/Foliage/Private/ProceduralFoliageComponent.cpp
+++ b/Engine/Source/Runtime/Foliage/Private/ProceduralFoliageComponent.cpp
@@ -91,24 +91,24 @@ UProceduralFoliageComponent::FGenerateProceduralContentParams UProceduralFoliage
 	return Params;
 }
 
-void UProceduralFoliageComponent::GetTileLayout(const FGenerateProceduralContentParams& Params, FTileLayout& OutTileLayout)
+void UProceduralFoliageComponent::GetTileLayout(const FGenerateProceduralContentParams& InParams, FTileLayout& OutTileLayout)
 {
-	if (Params.Bounds.IsValid)
+	if (InParams.Bounds.IsValid)
 	{
 		// Determine the bottom-left-most tile that contains the min position (when accounting for overlap)
-		const FVector MinPosition = Params.Bounds.Min + Params.TileOverlap;
-		OutTileLayout.BottomLeftX = FMath::FloorToInt(MinPosition.X / Params.FoliageSpawner->TileSize);
-		OutTileLayout.BottomLeftY = FMath::FloorToInt(MinPosition.Y / Params.FoliageSpawner->TileSize);
+		const FVector MinPosition = InParams.Bounds.Min + InParams.TileOverlap;
+		OutTileLayout.BottomLeftX = FMath::FloorToInt(MinPosition.X / InParams.FoliageSpawner->TileSize);
+		OutTileLayout.BottomLeftY = FMath::FloorToInt(MinPosition.Y / InParams.FoliageSpawner->TileSize);
 
 		// Determine the total number of tiles along each active axis
-		const FVector MaxPosition = Params.Bounds.Max - Params.TileOverlap;
-		const int32 MaxXIdx = FMath::FloorToInt(MaxPosition.X / Params.FoliageSpawner->TileSize);
-		const int32 MaxYIdx = FMath::FloorToInt(MaxPosition.Y / Params.FoliageSpawner->TileSize);
+		const FVector MaxPosition = InParams.Bounds.Max - InParams.TileOverlap;
+		const int32 MaxXIdx = FMath::FloorToInt(MaxPosition.X / InParams.FoliageSpawner->TileSize);
+		const int32 MaxYIdx = FMath::FloorToInt(MaxPosition.Y / InParams.FoliageSpawner->TileSize);
 
 		OutTileLayout.NumTilesX = (MaxXIdx - OutTileLayout.BottomLeftX) + 1;
 		OutTileLayout.NumTilesY = (MaxYIdx - OutTileLayout.BottomLeftY) + 1;
 
-		OutTileLayout.HalfHeight = Params.Bounds.GetExtent().Z;
+		OutTileLayout.HalfHeight = InParams.Bounds.GetExtent().Z;
 	}
 }
 
@@ -158,18 +158,18 @@ FVector UProceduralFoliageComponent::GetWorldPosition() const
 	return GetWorldPosition(GetGenerateProceduralContentParams());
 }
 
-FVector UProceduralFoliageComponent::GetWorldPosition(const FGenerateProceduralContentParams& Params)
+FVector UProceduralFoliageComponent::GetWorldPosition(const FGenerateProceduralContentParams& InParams)
 {
-	if (!Params.Bounds.IsValid || Params.FoliageSpawner == nullptr)
+	if (!InParams.Bounds.IsValid || InParams.FoliageSpawner == nullptr)
 	{
 		return FVector::ZeroVector;
 	}
 	
 	FTileLayout TileLayout;
-	GetTileLayout(Params, TileLayout);
+	GetTileLayout(InParams, TileLayout);
 
-	const float TileSize = Params.FoliageSpawner->TileSize;
-	return FVector(TileLayout.BottomLeftX * TileSize, TileLayout.BottomLeftY * TileSize, Params.Bounds.GetCenter().Z);
+	const float TileSize = InParams.FoliageSpawner->TileSize;
+	return FVector(TileLayout.BottomLeftX * TileSize, TileLayout.BottomLeftY * TileSize, InParams.Bounds.GetCenter().Z);
 }
 
 bool UProceduralFoliageComponent::GenerateProceduralContent(TArray<FDesiredFoliageInstance>& OutInstances)
@@ -182,10 +182,10 @@ bool UProceduralFoliageComponent::GenerateProceduralContent(TArray<FDesiredFolia
 	return false;
 }
 
-bool UProceduralFoliageComponent::GenerateProceduralContent(const FGenerateProceduralContentParams& Params, TArray<FDesiredFoliageInstance>& OutInstances)
+bool UProceduralFoliageComponent::GenerateProceduralContent(const FGenerateProceduralContentParams& InParams, TArray<FDesiredFoliageInstance>& OutInstances)
 {
 #if WITH_EDITOR
-	if (Params.FoliageSpawner)
+	if (InParams.FoliageSpawner)
 	{
 		// Establish the counter used to check if the user has canceled the simulation
 		FThreadSafeCounter LastCancel;
@@ -193,15 +193,15 @@ bool UProceduralFoliageComponent::GenerateProceduralContent(const FGenerateProce
 		FThreadSafeCounter* LastCancelPtr = &LastCancel;
 		
 		// Establish basic info about the tiles
-		const float TileOverlap = Params.TileOverlap;
-		const float TileSize = Params.FoliageSpawner->TileSize;
-		const FVector WorldPosition = GetWorldPosition(Params);
-		FVector2D GenerateLocation = FVector2D(Params.Bounds.GetCenter());
-		const float GenerateMaxExtent = FVector2D(Params.Bounds.GetExtent()).GetMax();
+		const float TileOverlap = InParams.TileOverlap;
+		const float TileSize = InParams.FoliageSpawner->TileSize;
+		const FVector WorldPosition = GetWorldPosition(InParams);
+		FVector2D GenerateLocation = FVector2D(InParams.Bounds.GetCenter());
+		const float GenerateMaxExtent = FVector2D(InParams.Bounds.GetExtent()).GetMax();
 		FTileLayout TileLayout;
-		GetTileLayout(Params, TileLayout);
+		GetTileLayout(InParams, TileLayout);
 
-		Params.FoliageSpawner->Simulate();
+		InParams.FoliageSpawner->Simulate();
 
 		TArray<TFuture< TArray<FDesiredFoliageInstance>* >> Futures;
 		for (int32 X = 0; X < TileLayout.NumTilesX; ++X)
@@ -209,7 +209,7 @@ bool UProceduralFoliageComponent::GenerateProceduralContent(const FGenerateProce
 			for (int32 Y = 0; Y < TileLayout.NumTilesY; ++Y)
 			{
 				// We have to get the simulated tiles and create new ones to build on main thread
-				const UProceduralFoliageTile* Tile = Params.FoliageSpawner->GetRandomTile(X + TileLayout.BottomLeftX, Y + TileLayout.BottomLeftY);
+				const UProceduralFoliageTile* Tile = InParams.FoliageSpawner->GetRandomTile(X + TileLayout.BottomLeftX, Y + TileLayout.BottomLeftY);
 				if (Tile == nullptr)	
 				{
 					// Simulation was either canceled or failed
@@ -217,12 +217,12 @@ bool UProceduralFoliageComponent::GenerateProceduralContent(const FGenerateProce
 				}
 
 				// From the pool of simulated tiles, pick the neighbors of this tile
-				const UProceduralFoliageTile* RightTile = (X + 1 < TileLayout.NumTilesX) ? Params.FoliageSpawner->GetRandomTile(X + TileLayout.BottomLeftX + 1, Y + TileLayout.BottomLeftY) : nullptr;
-				const UProceduralFoliageTile* TopTile = (Y + 1 < TileLayout.NumTilesY) ? Params.FoliageSpawner->GetRandomTile(X + TileLayout.BottomLeftX, Y + TileLayout.BottomLeftY + 1) : nullptr;
-				const UProceduralFoliageTile* TopRightTile = (RightTile && TopTile) ? Params.FoliageSpawner->GetRandomTile(X + TileLayout.BottomLeftX + 1, Y + TileLayout.BottomLeftY + 1) : nullptr;
+				const UProceduralFoliageTile* RightTile = (X + 1 < TileLayout.NumTilesX) ? InParams.FoliageSpawner->GetRandomTile(X + TileLayout.BottomLeftX + 1, Y + TileLayout.BottomLeftY) : nullptr;
+				const UProceduralFoliageTile* TopTile = (Y + 1 < TileLayout.NumTilesY) ? InParams.FoliageSpawner->GetRandomTile(X + TileLayout.BottomLeftX, Y + TileLayout.BottomLeftY + 1) : nullptr;
+				const UProceduralFoliageTile* TopRightTile = (RightTile && TopTile) ? InParams.FoliageSpawner->GetRandomTile(X + TileLayout.BottomLeftX + 1, Y + TileLayout.BottomLeftY + 1) : nullptr;
 
 				// Create a temp tile that will contain the composite contents of the tile after accounting for overlap
-				UProceduralFoliageTile* CompositeTile = Params.FoliageSpawner->CreateTempTile();
+				UProceduralFoliageTile* CompositeTile = InParams.FoliageSpawner->CreateTempTile();
 
 				Futures.Add(Async(EAsyncExecution::ThreadPool, [=]()
 				{
@@ -235,7 +235,7 @@ bool UProceduralFoliageComponent::GenerateProceduralContent(const FGenerateProce
 					//@todo proc foliage: Determine the composite contents of the tile (including overlaps) without copying everything to a temp tile
 
 					// Copy the base tile contents
-					const FBox2D BaseTile = GetTileRegion(X, Y, TileSize, Params.TileOverlap);
+					const FBox2D BaseTile = GetTileRegion(X, Y, TileSize, InParams.TileOverlap);
 					Tile->CopyInstancesToTile(CompositeTile, BaseTile, FTransform::Identity, TileOverlap);
 
 
@@ -267,7 +267,7 @@ bool UProceduralFoliageComponent::GenerateProceduralContent(const FGenerateProce
 					const FTransform TileTM(OrientedOffset + WorldPosition);
 
 					TArray<FDesiredFoliageInstance>* DesiredInstances = new TArray<FDesiredFoliageInstance>();
-					CompositeTile->ExtractDesiredInstances(*DesiredInstances, TileTM, GenerateLocation, GenerateMaxExtent, Params.ProceduralGuid, TileLayout.HalfHeight, Params.ProceduralVolumeInstance, true);
+					CompositeTile->ExtractDesiredInstances(*DesiredInstances, TileTM, GenerateLocation, GenerateMaxExtent, InParams.ProceduralGuid, TileLayout.HalfHeight, InParams.ProceduralVolumeInstance, true);
 
 					return DesiredInstances;
 					
@@ -371,20 +371,20 @@ bool UProceduralFoliageComponent::ResimulateProceduralFoliage(TFunctionRef<void(
 	return false;
 }
 
-void UProceduralFoliageComponent::RemoveProceduralContent(bool InRebuildTree)
+void UProceduralFoliageComponent::RemoveProceduralContent(bool bInRebuildTree)
 {
 #if WITH_EDITOR
 	TSet<AInstancedFoliageActor*> OutModifiedActors;
-	RemoveProceduralContent(GetWorld(), GetProceduralGuid(), InRebuildTree, OutModifiedActors);
+	RemoveProceduralContent(GetWorld(), GetProceduralGuid(), bInRebuildTree, OutModifiedActors);
 #endif
 }
 
-void UProceduralFoliageComponent::RemoveProceduralContent(UWorld* World, const FGuid& ProceduralGuid, bool InRebuildTree, TSet<AInstancedFoliageActor*>& OutModifiedActors)
+void UProceduralFoliageComponent::RemoveProceduralContent(UWorld* InWorld, const FGuid& InProceduralGuid, bool bInRebuildTree, TSet<AInstancedFoliageActor*>& OutModifiedActors)
 {
 #if WITH_EDITOR
-	for (TActorIterator<AInstancedFoliageActor> It(World); It; ++It)
+	for (TActorIterator<AInstancedFoliageActor> It(InWorld); It; ++It)
 	{
-		if ((*It)->DeleteInstancesForProceduralFoliageComponent(ProceduralGuid, InRebuildTree))
+		if ((*It)->DeleteInstancesForProceduralFoliageComponent(InProceduralGuid, bInRebuildTree))
 		{
 			OutModifiedActors.Add(*It);
 		}
diff --git a/Engine/Source/Runtime/Foliage/Public/InstancedFoliageActor.h b/Engine/Source/Runtime/Foliage/Public/InstancedFoliageActor.h
--- a/Engine/Source/Runtime/Foliage/Public/InstancedFoliageActor.h
+++ b/Engine/Source/Runtime/Foliage/Public/InstancedFoliageActor.h
@@ -219,12 +219,13 @@ public:
 	static FOLIAGE_API void DeleteInstancesForComponent(UWorld* InWorld, UActorComponent* InComponent);
 
 	// Deletes the instances spawned by a procedural component
-	FOLIAGE_API bool DeleteInstancesForProceduralFoliageComponent(const UProceduralFoliageComponent* ProceduralFoliageComponent, bool InRebuildTree);
-	FOLIAGE_API bool DeleteInstancesForProceduralFoliageComponent(const FGuid& ProceduralGuid, bool InRebuildTree);
-	
+	FOLIAGE_API bool DeleteInstancesForProceduralFoliageComponent(const UProceduralFoliageComponent* InProceduralFoliageComponent, bool bInRebuildTree);
+	FOLIAGE_API bool DeleteInstancesForProceduralFoliageComponent(const FGuid& InProceduralGuid, bool bInRebuildTree);
+	FOLIAGE_API bool DeleteInstancesForAllProceduralFoliageComponents(bool bInRebuildTree);
+
 	/** @return True if any instances exist that were spawned by the given procedural component */
-	FOLIAGE_API bool ContainsInstancesFromProceduralFoliageComponent(const UProceduralFoliageComponent* ProceduralFoliageComponent);
-	FOLIAGE_API bool ContainsInstancesFromProceduralFoliageComponent(const FGuid& ProceduralGuid);
+	FOLIAGE_API bool ContainsInstancesFromProceduralFoliageComponent(const UProceduralFoliageComponent* InProceduralFoliageComponent);
+	FOLIAGE_API bool ContainsInstancesFromProceduralFoliageComponent(const FGuid& InProceduralGuid);
 
 	// Finds a mesh entry or adds it if it doesn't already exist
 	FOLIAGE_API FFoliageInfo* FindOrAddMesh(UFoliageType* InType);
@@ -302,6 +303,7 @@ private:
 	void OnPostApplyLevelOffset(ULevel* InLevel, UWorld* InWorld, const FVector& InOffset, bool bWorldShift);
 	void OnPostWorldInitialization(UWorld* World, const UWorld::InitializationValues IVS);
 	void MoveInstancesToNewComponent(UPrimitiveComponent* InOldComponent, UPrimitiveComponent* InNewComponent, TFunctionRef<TArray<int32>(const FFoliageInfo&)> GetInstancesToMoveFunc);
+	bool DeleteInstancesForProceduralFoliageComponentInternal(const FGuid& InProceduralGuid, bool bInRebuildTree, bool bInDeleteAll);
 #endif
 private:
 #if WITH_EDITOR
diff --git a/Engine/Source/Runtime/Foliage/Public/ProceduralFoliageComponent.h b/Engine/Source/Runtime/Foliage/Public/ProceduralFoliageComponent.h
--- a/Engine/Source/Runtime/Foliage/Public/ProceduralFoliageComponent.h
+++ b/Engine/Source/Runtime/Foliage/Public/ProceduralFoliageComponent.h
@@ -115,11 +115,11 @@ class FOLIAGE_API UProceduralFoliageComponent : public UActorComponent
 	 * @return True if the simulation succeeded
 	 */
 	bool GenerateProceduralContent(TArray<FDesiredFoliageInstance>& OutInstances);
-	static bool GenerateProceduralContent(const FGenerateProceduralContentParams& Params, TArray<FDesiredFoliageInstance>& OutInstances);
+	static bool GenerateProceduralContent(const FGenerateProceduralContentParams& InParams, TArray<FDesiredFoliageInstance>& OutInstances);
 			
 	/** Removes all spawned foliage instances in the level that were spawned by this component */
-	void RemoveProceduralContent(bool InRebuildTree = true);
-	static void RemoveProceduralContent(UWorld* World, const FGuid& ProceduralGuid, bool InRebuildTree, TSet<AInstancedFoliageActor*>& OutModifiedActors);
+	void RemoveProceduralContent(bool bInRebuildTree = true);
+	static void RemoveProceduralContent(UWorld* InWorld, const FGuid& InProceduralGuid, bool bInRebuildTree, TSet<AInstancedFoliageActor*>& OutModifiedActors);
 
 	/** @return True if any foliage instances in the level were spawned by this component */
 	bool HasSpawnedAnyInstances();
