Changed the InstancedStaticMeshComponent component creation so that it has hit proxies for its instances by default when it is not dynamically or natively created.

This make the instances selectable in the editor at all time. Before that user needed to add a instance via the details panel to get that behavior. That was confusing and it didn't allow editor script to create those selectable instances.

Also override the OnComponentCreated function for the UGrassInstancedStaticMeshComponent because they don't need per instance hit proxies.

diff --git a/Engine/Source/Runtime/Engine/Classes/Components/InstancedStaticMeshComponent.h b/Engine/Source/Runtime/Engine/Classes/Components/InstancedStaticMeshComponent.h
--- a/Engine/Source/Runtime/Engine/Classes/Components/InstancedStaticMeshComponent.h
+++ b/Engine/Source/Runtime/Engine/Classes/Components/InstancedStaticMeshComponent.h
@@ -288,8 +288,10 @@ class ENGINE_API UInstancedStaticMeshComponent : public UStaticMeshComponent, pu
 
 	virtual void PostLoad() override;
 	virtual void OnRegister() override;
+	virtual void OnComponentCreated() override;
 	virtual bool SupportsRemoveSwap() const { return false; }
 
+
 #if WITH_EDITOR
 	virtual bool ComponentIsTouchingSelectionBox(const FBox& InSelBBox, const bool bConsiderOnlyBSP, const bool bMustEncompassEntireComponent) const override;
 	virtual bool ComponentIsTouchingSelectionFrustum(const FConvexVolume& InFrustum, const bool bConsiderOnlyBSP, const bool bMustEncompassEntireComponent) const override;
diff --git a/Engine/Source/Runtime/Engine/Private/InstancedStaticMesh.cpp b/Engine/Source/Runtime/Engine/Private/InstancedStaticMesh.cpp
--- a/Engine/Source/Runtime/Engine/Private/InstancedStaticMesh.cpp
+++ b/Engine/Source/Runtime/Engine/Private/InstancedStaticMesh.cpp
@@ -4255,6 +4255,19 @@ void UInstancedStaticMeshComponent::OnRegister()
 	}
 }
 
+void UInstancedStaticMeshComponent::OnComponentCreated()
+{
+	Super::OnComponentCreated();
+
+	if ( CreationMethod == EComponentCreationMethod::SimpleConstructionScript
+		|| CreationMethod == EComponentCreationMethod::Instance
+		)
+	{
+		// If it the component is selectable in the level editor make sure that we generate hit proxies so that its instance are also selectable
+		bHasPerInstanceHitProxies = true;
+	}
+}
+
 #if WITH_EDITOR
 
 bool UInstancedStaticMeshComponent::ComponentIsTouchingSelectionBox(const FBox& InSelBBox, const bool bConsiderOnlyBSP, const bool bMustEncompassEntireComponent) const
diff --git a/Engine/Source/Runtime/Foliage/Private/GrassInstancedStaticMeshComponent.cpp b/Engine/Source/Runtime/Foliage/Private/GrassInstancedStaticMeshComponent.cpp
--- a/Engine/Source/Runtime/Foliage/Private/GrassInstancedStaticMeshComponent.cpp
+++ b/Engine/Source/Runtime/Foliage/Private/GrassInstancedStaticMeshComponent.cpp
@@ -10,6 +10,14 @@ static TAutoConsoleVariable<bool> CVarEnableGrassInstancedWPOVelocity(
 	TEXT(" False")
 	);
 
+void UGrassInstancedStaticMeshComponent::OnComponentCreated()
+{
+	Super::OnComponentCreated();
+
+	// Grass instanced don't need hit proxies by default even when created in the editor manually by the user.
+	bHasPerInstanceHitProxies = false;
+}
+
 bool UGrassInstancedStaticMeshComponent::SupportsWorldPositionOffsetVelocity() const 
 { 
 	return CVarEnableGrassInstancedWPOVelocity.GetValueOnAnyThread();
diff --git a/Engine/Source/Runtime/Foliage/Public/GrassInstancedStaticMeshComponent.h b/Engine/Source/Runtime/Foliage/Public/GrassInstancedStaticMeshComponent.h
--- a/Engine/Source/Runtime/Foliage/Public/GrassInstancedStaticMeshComponent.h
+++ b/Engine/Source/Runtime/Foliage/Public/GrassInstancedStaticMeshComponent.h
@@ -12,6 +12,9 @@ class FOLIAGE_API UGrassInstancedStaticMeshComponent : public UHierarchicalInsta
 	GENERATED_UCLASS_BODY()
 
 public:
+
+	void OnComponentCreated() override;
+
 	static void BuildTreeAnyThread(TArray<FMatrix>& InstanceTransforms, TArray<float>& InstanceCustomDataFloats, int32 NumCustomDataFloats, const FBox& MeshBox, TArray<FClusterNode>& OutClusterTree, TArray<int32>& OutSortedInstances, TArray<int32>& OutInstanceReorderTable, int32& OutOcclusionLayerNum, int32 MaxInstancesPerLeaf, bool InGenerateInstanceScalingRange);
 	
 	void AcceptPrebuiltTree(TArray<FInstancedStaticMeshInstanceData>& InInstanceData, TArray<FClusterNode>& InClusterTree, int32 InOcclusionLayerNumNodes, int32 InNumBuiltRenderInstances);
