Foliage: Fix instances not being moved to different cell when moved through base component

diff --git a/Engine/Source/Editor/FoliageEdit/Private/FoliageEdMode.cpp b/Engine/Source/Editor/FoliageEdit/Private/FoliageEdMode.cpp
--- a/Engine/Source/Editor/FoliageEdit/Private/FoliageEdMode.cpp
+++ b/Engine/Source/Editor/FoliageEdit/Private/FoliageEdMode.cpp
@@ -1907,50 +1907,7 @@ void FEdModeFoliage::ApplySelection(UWorld* InWorld, bool bApply)
 void FEdModeFoliage::UpdateInstancePartitioning(UWorld* InWorld)
 {
 	check(!bMoving);
-	UActorPartitionSubsystem* ActorPartitionSubsystem = InWorld->GetSubsystem<UActorPartitionSubsystem>();
-	for (TActorIterator<AInstancedFoliageActor> It(InWorld); It; ++It)
-	{
-		AInstancedFoliageActor* IFA = *It;
-		IFA->ForEachFoliageInfo([InWorld, IFA, ActorPartitionSubsystem](UFoliageType* FoliageType, FFoliageInfo& FoliageInfo)
-		{
-			if (FoliageInfo.Type == EFoliageImplType::Actor && 
-				ActorPartitionSubsystem->IsLevelPartition())
-			{
-				return true; // Actors are handled through the Partitioning code
-			}
-
-			AInstancedFoliageActor* TargetIFA = nullptr;
-			bool bMovedInstances = false;
-			// Loop here because SelectedIndices will change on MoveInstancesToLevel so we need to process the modified remaining SelectedIndices
-			do
-			{
-				TargetIFA = nullptr;
-				bMovedInstances = false;
-				TSet<int32> InstancesToMove;
-
-				for (int32 SelectedInstanceIdx : FoliageInfo.SelectedIndices)
-				{
-					FFoliageInstance& Instance = FoliageInfo.Instances[SelectedInstanceIdx];
-					AInstancedFoliageActor* NewIFA = AInstancedFoliageActor::Get(InWorld, true, IFA->GetLevel(), Instance.Location);
-					if ((TargetIFA == nullptr || TargetIFA == NewIFA) && NewIFA != IFA)
-					{
-						TargetIFA = NewIFA;
-						InstancesToMove.Add(SelectedInstanceIdx);
-					}
-				}
-
-				if (InstancesToMove.Num())
-				{
-					// TargetIFA can be null (if target is unloaded cell). Meaning instances will be deleted.
-					FoliageInfo.MoveInstances(TargetIFA, InstancesToMove, true);
-					bMovedInstances = true;
-				}
-
-			} while (bMovedInstances && FoliageInfo.SelectedIndices.Num() > 0);
-
-			return true; // continue iteration
-		});
-	}
+	AInstancedFoliageActor::UpdateInstancePartitioning(InWorld);
 }
 
 void FEdModeFoliage::PostTransformSelectedInstances(UWorld* InWorld)
diff --git a/Engine/Source/Runtime/Foliage/Private/FoliageInstanceBase.cpp b/Engine/Source/Runtime/Foliage/Private/FoliageInstanceBase.cpp
--- a/Engine/Source/Runtime/Foliage/Private/FoliageInstanceBase.cpp
+++ b/Engine/Source/Runtime/Foliage/Private/FoliageInstanceBase.cpp
@@ -126,6 +126,7 @@ FFoliageInstanceBaseId FFoliageInstanceBaseCache::AddInstanceBaseId(UActorCompon
 	if (InComponent && !InComponent->IsCreatedByConstructionScript())
 	{
 		BaseId = GetInstanceBaseId(InComponent);
+		FFoliageInstanceBaseInfo BaseInfo(InComponent);
 		if (BaseId == FFoliageInstanceBaseCache::InvalidBaseId)
 		{
 			// generate next unique ID for base component
@@ -135,8 +136,6 @@ FFoliageInstanceBaseId FFoliageInstanceBaseCache::AddInstanceBaseId(UActorCompon
 			}
 			while (InstanceBaseMap.Contains(BaseId));
 			
-			FFoliageInstanceBaseInfo BaseInfo(InComponent);
-			
 			// more info for UE-30878
 			if (InstanceBaseInvMap.Contains(BaseInfo.BasePtr))
 			{
@@ -159,6 +158,11 @@ FFoliageInstanceBaseId FFoliageInstanceBaseCache::AddInstanceBaseId(UActorCompon
 				}
 			}
 		}
+		else
+		{
+			// Update BaseInfo
+			InstanceBaseMap[BaseId] = MoveTemp(BaseInfo);
+		}
 	}
 
 	return BaseId;
diff --git a/Engine/Source/Runtime/Foliage/Private/InstancedFoliage.cpp b/Engine/Source/Runtime/Foliage/Private/InstancedFoliage.cpp
--- a/Engine/Source/Runtime/Foliage/Private/InstancedFoliage.cpp
+++ b/Engine/Source/Runtime/Foliage/Private/InstancedFoliage.cpp
@@ -3134,6 +3134,62 @@ void AInstancedFoliageActor::MoveInstancesForMovedOwnedActors(AActor* InActor)
 	}
 }
 
+namespace FoliagePartitioningUtils
+{
+	static void Update(AInstancedFoliageActor* SourceIFA, FFoliageInfo& FoliageInfo, TFunctionRef<TSet<int32>* ()> GetInstanceSet)
+	{
+		UActorPartitionSubsystem* ActorPartitionSubsystem = SourceIFA->GetWorld()->GetSubsystem<UActorPartitionSubsystem>();
+		if (FoliageInfo.Type == EFoliageImplType::Actor && ActorPartitionSubsystem->IsLevelPartition())
+		{
+			return; // Actors are handled through the Partitioning code
+		}
+
+		bool bMovedInstances = true;
+		AInstancedFoliageActor* TargetIFA = nullptr;
+
+		TSet<int32>* InstanceSet = GetInstanceSet();
+		while (InstanceSet && InstanceSet->Num() > 0 && bMovedInstances)
+		{
+			TargetIFA = nullptr;
+			bMovedInstances = false;
+			TSet<int32> InstancesToMove;
+
+			for (int32 InstanceIdx : *InstanceSet)
+			{
+				FFoliageInstance& Instance = FoliageInfo.Instances[InstanceIdx];
+				AInstancedFoliageActor* NewIFA = AInstancedFoliageActor::Get(SourceIFA->GetWorld(), true, SourceIFA->GetLevel(), Instance.Location);
+				if ((TargetIFA == nullptr || TargetIFA == NewIFA) && NewIFA != SourceIFA)
+				{
+					TargetIFA = NewIFA;
+					InstancesToMove.Add(InstanceIdx);
+				}
+			}
+
+			if (InstancesToMove.Num())
+			{
+				// TargetIFA can be null (if target is unloaded cell). Meaning instances will be deleted.
+				FoliageInfo.MoveInstances(TargetIFA, InstancesToMove, true);
+				bMovedInstances = true;
+			}
+
+			InstanceSet = GetInstanceSet();
+		}
+	}
+}
+
+void AInstancedFoliageActor::UpdateInstancePartitioning(UWorld* InWorld)
+{
+	for (TActorIterator<AInstancedFoliageActor> It(InWorld); It; ++It)
+	{
+		AInstancedFoliageActor* IFA = *It;
+		IFA->ForEachFoliageInfo([IFA](UFoliageType* FoliageType, FFoliageInfo& FoliageInfo)
+        {
+			FoliagePartitioningUtils::Update(IFA, FoliageInfo, [&]() { return &FoliageInfo.SelectedIndices; });
+			return true; // continue iteration
+        });
+	}
+}
+
 void AInstancedFoliageActor::MoveInstancesForMovedComponent(UActorComponent* InComponent)
 {
 	// We don't want to handle this case when applying level transform
@@ -3206,6 +3262,8 @@ void AInstancedFoliageActor::MoveInstancesForMovedComponent(UActorComponent* InC
 
 			Info.Implementation->EndUpdate();
 			Info.Refresh(true, false);
+
+			FoliagePartitioningUtils::Update(this, Info, [&]() { return Info.ComponentHash.Find(BaseId); });
 		}
 	}
 }
diff --git a/Engine/Source/Runtime/Foliage/Public/InstancedFoliageActor.h b/Engine/Source/Runtime/Foliage/Public/InstancedFoliageActor.h
--- a/Engine/Source/Runtime/Foliage/Public/InstancedFoliageActor.h
+++ b/Engine/Source/Runtime/Foliage/Public/InstancedFoliageActor.h
@@ -170,6 +170,8 @@ public:
 
 	static FOLIAGE_API AInstancedFoliageActor* GetDefault(UWorld* InWorld);
 
+	static FOLIAGE_API void UpdateInstancePartitioning(UWorld* InWorld);
+
 	static FOLIAGE_API bool FoliageTrace(const UWorld* InWorld, FHitResult& OutHit, const FDesiredFoliageInstance& DesiredInstance, FName InTraceTag = NAME_None, bool InbReturnFaceIndex = false, const FFoliageTraceFilterFunc& FilterFunc = FFoliageTraceFilterFunc(), bool bAverageNormal = false);
 	static FOLIAGE_API bool CheckCollisionWithWorld(const UWorld* InWorld, const UFoliageType* Settings, const FFoliageInstance& Inst, const FVector& HitNormal, const FVector& HitLocation, UPrimitiveComponent* HitComponent);
 
