Fix Foliage deprecation crash

diff --git a/Engine/Source/Runtime/Foliage/Private/InstancedFoliage.cpp b/Engine/Source/Runtime/Foliage/Private/InstancedFoliage.cpp
--- a/Engine/Source/Runtime/Foliage/Private/InstancedFoliage.cpp
+++ b/Engine/Source/Runtime/Foliage/Private/InstancedFoliage.cpp
@@ -311,8 +311,28 @@ static void ConvertDeprecated2FoliageMeshes(
 #if WITH_EDITORONLY_DATA	
 	for (auto Pair : FoliageMeshesDeprecated)
 	{
+		UFoliageType* FoliageType = Pair.Key;
+		TUniqueObj<FFoliageMeshInfo_Deprecated2>& FoliageMeshDeprecated = Pair.Value;
+		
+		if (!FoliageType)
+		{
+			FMessageLog("MapCheck").Warning()
+				->AddToken(FUObjectToken::Create(IFA))
+				->AddToken(FTextToken::Create(LOCTEXT("MapCheck_Message_FoliageMissingFoliageType", "Foliage instances for a missing foliage type have been removed.")))
+				->AddToken(FMapErrorToken::Create(FMapErrors::FoliageMissingStaticMesh));
+
+			if (FoliageMeshDeprecated->Component)
+			{
+				FoliageMeshDeprecated->Component->SetFlags(RF_Transactional);
+				FoliageMeshDeprecated->Component->Modify();
+				FoliageMeshDeprecated->Component->DestroyComponent();
+				FoliageMeshDeprecated->Component = nullptr;
+			}
+
+			continue;
+		}
+
 		auto& FoliageMesh = IFA->AddFoliageInfo(Pair.Key);
-		const auto& FoliageMeshDeprecated = Pair.Value;
 
 		// Old Foliage mesh is always static mesh (no actors)
 		FoliageMesh->Type = EFoliageImplType::StaticMesh;
