Foliage: In World Partition use fixed bounds (fixes empty foliage actors not loading with loading regions because of null extents)

diff --git a/Engine/Source/Runtime/Foliage/Private/InstancedFoliage.cpp b/Engine/Source/Runtime/Foliage/Private/InstancedFoliage.cpp
--- a/Engine/Source/Runtime/Foliage/Private/InstancedFoliage.cpp
+++ b/Engine/Source/Runtime/Foliage/Private/InstancedFoliage.cpp
@@ -4114,6 +4114,17 @@ bool AInstancedFoliageActor::ShouldImport(FString* ActorPropString, bool IsMovin
 	return false;
 }
 
+FBox AInstancedFoliageActor::GetStreamingBounds() const
+{
+	if (UWorld* World = GetWorld(); World && World->IsPartitionedWorld())
+	{
+		const float HalfGridSize = GridSize * 0.5f;
+		return FBox(GetActorLocation() - HalfGridSize, GetActorLocation() + HalfGridSize);
+	}
+
+	return Super::GetStreamingBounds();
+}
+
 void AInstancedFoliageActor::ApplySelection(bool bApply)
 {
 	for (auto& Pair : FoliageInfos)
diff --git a/Engine/Source/Runtime/Foliage/Public/InstancedFoliageActor.h b/Engine/Source/Runtime/Foliage/Public/InstancedFoliageActor.h
--- a/Engine/Source/Runtime/Foliage/Public/InstancedFoliageActor.h
+++ b/Engine/Source/Runtime/Foliage/Public/InstancedFoliageActor.h
@@ -180,6 +180,7 @@ public:
 	virtual void PostDuplicate(bool bDuplicateForPIE) override;
 	virtual bool ShouldExport() override;
 	virtual bool ShouldImport(FString* ActorPropString, bool IsMovingLevel) override;
+	virtual FBox GetStreamingBounds() const override;
 
 	// Called in response to BSP rebuilds to migrate foliage from obsolete to new components.
 	FOLIAGE_API void MapRebuild();
