Fix potential race condition by moving delegate registration out of the async loading thread for InstancedFoliageActor

diff --git a/Engine/Source/Runtime/Foliage/Private/InstancedFoliage.cpp b/Engine/Source/Runtime/Foliage/Private/InstancedFoliage.cpp
--- a/Engine/Source/Runtime/Foliage/Private/InstancedFoliage.cpp
+++ b/Engine/Source/Runtime/Foliage/Private/InstancedFoliage.cpp
@@ -4370,52 +4370,73 @@ void AInstancedFoliageActor::ExitEditMode()
 	}
 }
 
-void AInstancedFoliageActor::PostInitProperties()
+void AInstancedFoliageActor::RegisterDelegates()
 {
-	Super::PostInitProperties();
+	ensureMsgf(IsInGameThread(), TEXT("Potential race condition in AInstancedFoliageActor registering to delegates from non game-thread"));
 
-	if (!IsTemplate())
+	GEngine->OnActorMoved().Remove(OnLevelActorMovedDelegateHandle);
+	OnLevelActorMovedDelegateHandle = GEngine->OnActorMoved().AddUObject(this, &AInstancedFoliageActor::OnLevelActorMoved);
+
+	GEngine->OnLevelActorDeleted().Remove(OnLevelActorDeletedDelegateHandle);
+	OnLevelActorDeletedDelegateHandle = GEngine->OnLevelActorDeleted().AddUObject(this, &AInstancedFoliageActor::OnLevelActorDeleted);
+
+	if (GetLevel())
 	{
-		GEngine->OnActorMoved().Remove(OnLevelActorMovedDelegateHandle);
-		OnLevelActorMovedDelegateHandle = GEngine->OnActorMoved().AddUObject(this, &AInstancedFoliageActor::OnLevelActorMoved);
+		OnApplyLevelTransformDelegateHandle = GetLevel()->OnApplyLevelTransform.AddUObject(this, &AInstancedFoliageActor::OnApplyLevelTransform);
+	}
 
-		GEngine->OnLevelActorDeleted().Remove(OnLevelActorDeletedDelegateHandle);
-		OnLevelActorDeletedDelegateHandle = GEngine->OnLevelActorDeleted().AddUObject(this, &AInstancedFoliageActor::OnLevelActorDeleted);
+	GEngine->OnLevelActorOuterChanged().Remove(OnLevelActorOuterChangedDelegateHandle);
+	OnLevelActorOuterChangedDelegateHandle = GEngine->OnLevelActorOuterChanged().AddUObject(this, &AInstancedFoliageActor::OnLevelActorOuterChanged);
 
-		if (GetLevel())
-		{
-			OnApplyLevelTransformDelegateHandle = GetLevel()->OnApplyLevelTransform.AddUObject(this, &AInstancedFoliageActor::OnApplyLevelTransform);
-		}
+	FWorldDelegates::PostApplyLevelOffset.Remove(OnPostApplyLevelOffsetDelegateHandle);
+	OnPostApplyLevelOffsetDelegateHandle = FWorldDelegates::PostApplyLevelOffset.AddUObject(this, &AInstancedFoliageActor::OnPostApplyLevelOffset);
+
+	FWorldDelegates::OnPostWorldInitialization.Remove(OnPostWorldInitializationDelegateHandle);
+	OnPostWorldInitializationDelegateHandle = FWorldDelegates::OnPostWorldInitialization.AddUObject(this, &AInstancedFoliageActor::OnPostWorldInitialization);
+}
 
-		GEngine->OnLevelActorOuterChanged().Remove(OnLevelActorOuterChangedDelegateHandle);
-		OnLevelActorOuterChangedDelegateHandle = GEngine->OnLevelActorOuterChanged().AddUObject(this, &AInstancedFoliageActor::OnLevelActorOuterChanged);
+void AInstancedFoliageActor::UnregisterDelegates()
+{
+	ensureMsgf(IsInGameThread(), TEXT("Potential race condition in AInstancedFoliageActor registering to delegates from non game-thread"));
 
-		FWorldDelegates::PostApplyLevelOffset.Remove(OnPostApplyLevelOffsetDelegateHandle);
-		OnPostApplyLevelOffsetDelegateHandle = FWorldDelegates::PostApplyLevelOffset.AddUObject(this, &AInstancedFoliageActor::OnPostApplyLevelOffset);
+	GEngine->OnActorMoved().Remove(OnLevelActorMovedDelegateHandle);
+	GEngine->OnLevelActorDeleted().Remove(OnLevelActorDeletedDelegateHandle);
+	GEngine->OnLevelActorOuterChanged().Remove(OnLevelActorOuterChangedDelegateHandle);
 
-		FWorldDelegates::OnPostWorldInitialization.Remove(OnPostWorldInitializationDelegateHandle);
-		OnPostWorldInitializationDelegateHandle = FWorldDelegates::OnPostWorldInitialization.AddUObject(this, &AInstancedFoliageActor::OnPostWorldInitialization);
+	if (GetLevel())
+	{
+		GetLevel()->OnApplyLevelTransform.Remove(OnApplyLevelTransformDelegateHandle);
 	}
+
+	FWorldDelegates::PostApplyLevelOffset.Remove(OnPostApplyLevelOffsetDelegateHandle);
+
+	FWorldDelegates::OnPostWorldInitialization.Remove(OnPostWorldInitializationDelegateHandle);
 }
 
-void AInstancedFoliageActor::BeginDestroy()
+void AInstancedFoliageActor::PostInitProperties()
 {
-	Super::BeginDestroy();
+	Super::PostInitProperties();
 
 	if (!IsTemplate())
 	{
-		GEngine->OnActorMoved().Remove(OnLevelActorMovedDelegateHandle);
-		GEngine->OnLevelActorDeleted().Remove(OnLevelActorDeletedDelegateHandle);
-		GEngine->OnLevelActorOuterChanged().Remove(OnLevelActorOuterChangedDelegateHandle);
-
-		if (GetLevel())
+		if (HasAnyFlags(RF_NeedPostLoad) || GetOuter()->HasAnyFlags(RF_NeedPostLoad))
 		{
-			GetLevel()->OnApplyLevelTransform.Remove(OnApplyLevelTransformDelegateHandle);
+			// Delegate registration is not thread-safe, so we postpone it on PostLoad when coming from loading which could be on another thread
 		}
+		else
+		{
+			RegisterDelegates();
+		}
+	}
+}
 
-		FWorldDelegates::PostApplyLevelOffset.Remove(OnPostApplyLevelOffsetDelegateHandle);
+void AInstancedFoliageActor::BeginDestroy()
+{
+	Super::BeginDestroy();
 
-		FWorldDelegates::OnPostWorldInitialization.Remove(OnPostWorldInitializationDelegateHandle);
+	if (!IsTemplate())
+	{
+		UnregisterDelegates();
 
 		FoliageInfos.Empty();
 	}
@@ -4433,6 +4454,13 @@ void AInstancedFoliageActor::PostLoad()
 {
 	Super::PostLoad();
 
+#if WITH_EDITOR
+	if (!IsTemplate())
+	{
+		RegisterDelegates();
+	}
+#endif
+
 	ULevel* OwningLevel = GetLevel();
 	// We can't check the ActorPartitionSubsystem here because World is not initialized yet. So we fallback on the bIsPartitioned
 	// to know if multiple InstanceFoliageActors is valid or not.
diff --git a/Engine/Source/Runtime/Foliage/Public/InstancedFoliageActor.h b/Engine/Source/Runtime/Foliage/Public/InstancedFoliageActor.h
--- a/Engine/Source/Runtime/Foliage/Public/InstancedFoliageActor.h
+++ b/Engine/Source/Runtime/Foliage/Public/InstancedFoliageActor.h
@@ -80,6 +80,9 @@ public:
 
 protected:
 #if WITH_EDITOR
+	void RegisterDelegates();
+	void UnregisterDelegates();
+
 	void HandleFoliageInstancePreMove(const FFoliageInstanceId& InstanceId);
 	void HandleFoliageInstancePostMove(const FFoliageInstanceId& InstanceId);
 #endif
