Changed FoliageType bAffectDynamicIndirectLighting default to true.  This property was previously only used by LPV, which has been removed, so the previous settings can be lost.

diff --git a/Engine/Source/Runtime/Foliage/Private/InstancedFoliage.cpp b/Engine/Source/Runtime/Foliage/Private/InstancedFoliage.cpp
--- a/Engine/Source/Runtime/Foliage/Private/InstancedFoliage.cpp
+++ b/Engine/Source/Runtime/Foliage/Private/InstancedFoliage.cpp
@@ -551,7 +551,7 @@ UFoliageType::UFoliageType(const FObjectInitializer& ObjectInitializer)
 	bCastDynamicShadow = true;
 	bCastStaticShadow = true;
 	bCastContactShadow = true;
-	bAffectDynamicIndirectLighting = false;
+	bAffectDynamicIndirectLighting = true;
 	// Most of the high instance count foliage like grass causes performance problems with distance field lighting
 	bAffectDistanceFieldLighting = false;
 	bCastShadowAsTwoSided = false;
