Foliage: FoliageType now have a bEnableCullDistanceScale property, used with the new cvar 'Foliage.CullDistanceScale' allows scaling of cull distance per platform

diff --git a/Engine/Source/Runtime/Foliage/Private/InstancedFoliage.cpp b/Engine/Source/Runtime/Foliage/Private/InstancedFoliage.cpp
--- a/Engine/Source/Runtime/Foliage/Private/InstancedFoliage.cpp
+++ b/Engine/Source/Runtime/Foliage/Private/InstancedFoliage.cpp
@@ -44,6 +44,7 @@ InstancedFoliage.cpp: Instanced foliage implementation.
 #include "DrawDebugHelpers.h"
 #include "UObject/UE5MainStreamObjectVersion.h"
 #include "UObject/FortniteMainBranchObjectVersion.h"
+#include "UObject/UObjectIterator.h"
 #include "PreviewScene.h"
 #include "FoliageActor.h"
 #include "FoliageISMActor.h"
@@ -75,6 +76,12 @@ static TAutoConsoleVariable<int32> CVarFoliageDiscardDataOnLoad(
 	TEXT("1: Discard foliage data on load if the foliage type has it enabled; 0: Keep foliage data regardless of whether the foliage type has it enabled or not (requires reloading level)"),
 	ECVF_Scalability);
 
+static TAutoConsoleVariable<float> CVarFoliageCullDistanceScale(
+	TEXT("foliage.CullDistanceScale"),
+	1.0,
+	TEXT("Controls the cull distance scale. Foliage must opt-in to cull distance scaling through the foliage type."),
+	ECVF_Scalability);
+
 const FGuid FFoliageCustomVersion::GUID(0x430C4D19, 0x71544970, 0x87699B69, 0xDF90B0E5);
 // Register the custom version with core
 FCustomVersionRegistration GRegisterFoliageCustomVersion(FFoliageCustomVersion::GUID, FFoliageCustomVersion::LatestVersion, TEXT("FoliageVer"));
@@ -194,6 +201,62 @@ private:
 #endif
 };
 
+namespace FoliageUtil
+{
+	FInt32Interval GetCullDistance(const UFoliageType* FoliageType, float CullDistanceScale)
+	{
+		if (!FoliageType->bEnableCullDistanceScaling)
+		{
+			return FoliageType->CullDistance;
+		}
+
+		return FInt32Interval(FoliageType->CullDistance.Min * CullDistanceScale, FoliageType->CullDistance.Max * CullDistanceScale);
+	}
+
+	void UpdateComponentCullDistance(const UFoliageType* FoliageType, UHierarchicalInstancedStaticMeshComponent* Component, float CullDistanceScale)
+	{
+		const FInt32Interval ScaledCullDistance = GetCullDistance(FoliageType, CullDistanceScale);
+
+		Component->SetCullDistances(ScaledCullDistance.Min, ScaledCullDistance.Max);
+	}
+
+#if WITH_EDITOR
+	static void CVarSinkFunction()
+	{
+		static float CachedCullDistanceScale = 1.0f;
+		float CullDistanceScale = CVarFoliageCullDistanceScale.GetValueOnGameThread();
+
+		if (CullDistanceScale != CachedCullDistanceScale)
+		{
+			CachedCullDistanceScale = CullDistanceScale;
+			CullDistanceScale = FMath::Clamp(CullDistanceScale, 0.0f, 1.0f);
+
+			for (const AInstancedFoliageActor* IFA : TObjectRange<AInstancedFoliageActor>(RF_ClassDefaultObject | RF_ArchetypeObject, true, EInternalObjectFlags::Garbage))
+			{
+				for (const TPair<UFoliageType*, TUniqueObj<FFoliageInfo>>& Pair : IFA->GetFoliageInfos())
+				{
+					if (Pair.Key && Pair.Key->bEnableCullDistanceScaling)
+					{
+						if (Pair.Value->Type == EFoliageImplType::StaticMesh)
+						{
+							FFoliageStaticMesh* FoliageStaticMesh = StaticCast<FFoliageStaticMesh*>(Pair.Value->Implementation.Get());
+
+							if (FoliageStaticMesh->Component != nullptr)
+							{
+								FoliageUtil::UpdateComponentCullDistance(Pair.Key, FoliageStaticMesh->Component, CullDistanceScale);
+							}
+						}
+					}
+				}
+			}
+		}
+	}
+
+	static FAutoConsoleVariableSink CVarSink(FConsoleCommandDelegate::CreateStatic(&CVarSinkFunction));
+
+#endif
+}
+
 struct FFoliagePlacementUtil
 {
 	static int32 GetRandomSeedForPosition(const FVector2D& Position)
@@ -594,6 +657,7 @@ UFoliageType::UFoliageType(const FObjectInitializer& ObjectInitializer)
 #endif
 	bEnableDensityScaling = false;
 	bEnableDiscardOnLoad = false;
+	bEnableCullDistanceScaling = false;
 
 #if WITH_EDITORONLY_DATA
 	bIncludeInHLOD = true;
@@ -1536,16 +1600,10 @@ void FFoliageStaticMesh::UpdateComponentSettings(const UFoliageType_InstancedSta
 			bNeedsMarkRenderStateDirty = true;
 			bNeedsInvalidateLightingCache = true;
 		}
-		if (Component->InstanceStartCullDistance != FoliageType->CullDistance.Min)
-		{
-			Component->InstanceStartCullDistance = FoliageType->CullDistance.Min;
-			bNeedsMarkRenderStateDirty = true;
-		}
-		if (Component->InstanceEndCullDistance != FoliageType->CullDistance.Max)
-		{
-			Component->InstanceEndCullDistance = FoliageType->CullDistance.Max;
-			bNeedsMarkRenderStateDirty = true;
-		}
+
+		const float FoliageCullDistanceScale = FMath::Clamp(CVarFoliageCullDistanceScale.GetValueOnGameThread(), 0.0f, 1.0f);
+		FoliageUtil::UpdateComponentCullDistance(FoliageType, Component, FoliageCullDistanceScale);
+
 		if (Component->CastShadow != FoliageType->CastShadow)
 		{
 			Component->CastShadow = FoliageType->CastShadow;
@@ -4765,12 +4823,15 @@ void AInstancedFoliageActor::PostLoad()
 
 #endif// WITH_EDITOR
 
-	if (!GIsEditor && CVarFoliageDiscardDataOnLoad.GetValueOnGameThread())
+	if (!GIsEditor)
 	{
+		const bool bFoliageDiscardOnLoad = !!CVarFoliageDiscardDataOnLoad.GetValueOnGameThread();
+		const float FoliageCullDistanceScale = FMath::Clamp(CVarFoliageCullDistanceScale.GetValueOnGameThread(), 0.0f, 1.0f);
+
 		bool bHasISMFoliage = false;
 		for (auto& Pair : FoliageInfos)
 		{
-			if (!Pair.Key || Pair.Key->bEnableDiscardOnLoad)
+			if (bFoliageDiscardOnLoad && (!Pair.Key || Pair.Key->bEnableDiscardOnLoad))
 			{
 				if (Pair.Value->Type == EFoliageImplType::StaticMesh)
 				{
@@ -4792,6 +4853,19 @@ void AInstancedFoliageActor::PostLoad()
 					bHasISMFoliage = true;
 				}
 			}
+			else if (Pair.Key && Pair.Key->bEnableCullDistanceScaling)
+			{
+				if (Pair.Value->Type == EFoliageImplType::StaticMesh)
+				{
+					FFoliageStaticMesh* FoliageStaticMesh = StaticCast<FFoliageStaticMesh*>(Pair.Value->Implementation.Get());
+
+					if (FoliageStaticMesh->Component != nullptr)
+					{
+						FoliageStaticMesh->Component->ConditionalPostLoad();
+						FoliageUtil::UpdateComponentCullDistance(Pair.Key, FoliageStaticMesh->Component, FoliageCullDistanceScale);
+					}
+				}
+			}
 				
 			Pair.Value = FFoliageInfo();
 		}
diff --git a/Engine/Source/Runtime/Foliage/Public/FoliageType.h b/Engine/Source/Runtime/Foliage/Public/FoliageType.h
--- a/Engine/Source/Runtime/Foliage/Public/FoliageType.h
+++ b/Engine/Source/Runtime/Foliage/Public/FoliageType.h
@@ -581,6 +581,12 @@ public:
 	UPROPERTY(EditAnywhere, Category=Scalability)
 	uint32 bEnableDiscardOnLoad:1;
 
+	/*
+	 * Whether this foliage type should be affected by the Engine's "foliage.CullDistanceScale" setting 
+	 */
+	UPROPERTY(EditAnywhere, Category = Scalability)
+	uint32 bEnableCullDistanceScaling : 1;
+
 public:
 	// VIRTUAL TEXTURE
 
