Foliage & FoliageEdit modules: UnsafeTypeCastWarningLevel

#rnx

diff --git a/Engine/Source/Editor/DetailCustomizations/Public/Customizations/MobilityCustomization.h b/Engine/Source/Editor/DetailCustomizations/Public/Customizations/MobilityCustomization.h
--- a/Engine/Source/Editor/DetailCustomizations/Public/Customizations/MobilityCustomization.h
+++ b/Engine/Source/Editor/DetailCustomizations/Public/Customizations/MobilityCustomization.h
@@ -20,7 +20,7 @@ enum class ECheckBoxState : uint8;
 class DETAILCUSTOMIZATIONS_API FMobilityCustomization : public IDetailCustomNodeBuilder, public TSharedFromThis<FMobilityCustomization>
 {
 public:
-	enum
+	enum : uint8
 	{
 		StaticMobilityBitMask = (1u << EComponentMobility::Static),
 		StationaryMobilityBitMask = (1u << EComponentMobility::Stationary),
diff --git a/Engine/Source/Editor/FoliageEdit/Private/FoliageEdMode.cpp b/Engine/Source/Editor/FoliageEdit/Private/FoliageEdMode.cpp
--- a/Engine/Source/Editor/FoliageEdit/Private/FoliageEdMode.cpp
+++ b/Engine/Source/Editor/FoliageEdit/Private/FoliageEdMode.cpp
@@ -1012,7 +1012,7 @@ void FEdModeFoliage::GetRandomVectorInBrush(FVector& OutStart, FVector& OutEnd)
 	OutEnd = BrushLocation + UISettings.GetRadius() * (Point - Rw);
 }
 
-static bool IsWithinSlopeAngle(float NormalZ, float MinAngle, float MaxAngle, float Tolerance = SMALL_NUMBER)
+static bool IsWithinSlopeAngle(FVector::FReal NormalZ, float MinAngle, float MaxAngle, float Tolerance = SMALL_NUMBER)
 {
 	const float MaxNormalAngle = FMath::Cos(FMath::DegreesToRadians(MaxAngle));
 	const float MinNormalAngle = FMath::Cos(FMath::DegreesToRadians(MinAngle));
@@ -1023,7 +1023,8 @@ static bool IsWithinSlopeAngle(float NormalZ, float MinAngle, float MaxAngle, fl
 static bool CheckLocationForPotentialInstance_ThreadSafe(const UFoliageType* Settings, const FVector& Location, const FVector& Normal)
 {
 	// Check height range
-	if (!Settings->Height.Contains(Location.Z))
+	FDoubleInterval HeightInterval((double)Settings->Height.Min, (double)Settings->Height.Max);
+	if (!HeightInterval.Contains(Location.Z))
 	{
 		return false;
 	}
@@ -2161,8 +2162,8 @@ void FEdModeFoliage::SelectInvalidInstances(const UFoliageType* Settings)
 					if (HitComponent == CurrentInstanceBase)
 					{
 						FVector InstanceWorldZOffset = Instance.GetInstanceWorldTransform().TransformVector(FVector(0.f, 0.f, Instance.ZOffset));
-						float DistanceToGround = FVector::Dist(Instance.Location, Hit.Location + InstanceWorldZOffset);
-						float InstanceWorldTreshold = Instance.GetInstanceWorldTransform().TransformVector(FVector(0.f, 0.f, InstanceOffGroundLocalThreshold)).Size();
+						FVector::FReal DistanceToGround = FVector::Dist(Instance.Location, Hit.Location + InstanceWorldZOffset);
+						FVector::FReal InstanceWorldTreshold = Instance.GetInstanceWorldTransform().TransformVector(FVector(0.f, 0.f, InstanceOffGroundLocalThreshold)).Size();
 
 						if ((DistanceToGround - InstanceWorldTreshold) <= KINDA_SMALL_NUMBER)
 						{
@@ -2245,6 +2246,7 @@ void FEdModeFoliage::ReapplyInstancesForBrush(UWorld* InWorld, AInstancedFoliage
 	TArray<int32> ExistingInstances;
 	FoliageInfo->GetInstancesInsideSphere(BrushSphere, ExistingInstances);
 
+	const FDoubleInterval HeightInterval(Settings->Height.Min, Settings->Height.Max);
 	bool bUpdated = false;
 	TArray<int32> UpdatedInstances;
 	TSet<int32> InstancesToDelete;
@@ -2352,13 +2354,13 @@ void FEdModeFoliage::ReapplyInstancesForBrush(UWorld* InWorld, AInstancedFoliage
 			// Reapply scale
 			if (Settings->ReapplyScaling)
 			{
-				FVector NewScale = Settings->GetRandomScale();
+				FVector3f NewScale = Settings->GetRandomScale();
 
 				if (Settings->ReapplyScaleX)
 				{
 					if (Settings->Scaling == EFoliageScaling::Uniform)
 					{
-						Instance.DrawScale3D = (FVector3f)NewScale;
+						Instance.DrawScale3D = NewScale;
 					}
 					else
 					{
@@ -2471,7 +2473,7 @@ void FEdModeFoliage::ReapplyInstancesForBrush(UWorld* InWorld, AInstancedFoliage
 			// Cull instances that don't meet the height range
 			if (Settings->ReapplyHeight)
 			{
-				if (!Settings->Height.Contains(Instance.Location.Z))
+				if (!HeightInterval.Contains(Instance.Location.Z))
 				{
 					InstancesToDelete.Add(InstanceIndex);
 					if (bReapplyLocation)
@@ -2714,7 +2716,7 @@ struct FFoliagePaintBucketTriangle
 		VertexColor[2] = InColor2;
 
 		WorldNormal = InLocalToWorld.GetDeterminant() >= 0.f ? Vector2 ^ Vector1 : Vector1 ^ Vector2;
-		float WorldNormalSize = WorldNormal.Size();
+		FVector::FReal WorldNormalSize = WorldNormal.Size();
 		Area = WorldNormalSize * 0.5f;
 		if (WorldNormalSize > SMALL_NUMBER)
 		{
@@ -2743,7 +2745,7 @@ struct FFoliagePaintBucketTriangle
 	FVector Vector1;
 	FVector Vector2;
 	FVector WorldNormal;
-	float Area;
+	FVector::FReal Area;
 	FColor VertexColor[3];
 };
 
@@ -2828,9 +2830,9 @@ void FEdModeFoliage::ApplyPaintBucket_Add(AActor* Actor)
 					const int32 Index1 = Indices[Idx + 1];
 					const int32 Index2 = Indices[Idx + 2];
 
-					const FVector Vert0 = SplineMesh->CalcSliceTransform(USplineMeshComponent::GetAxisValue((FVector)PositionVertexBuffer.VertexPosition(Index0), SplineMesh->ForwardAxis)).TransformPosition((FVector)PositionVertexBuffer.VertexPosition(Index0) * Mask);
-					const FVector Vert1 = SplineMesh->CalcSliceTransform(USplineMeshComponent::GetAxisValue((FVector)PositionVertexBuffer.VertexPosition(Index1), SplineMesh->ForwardAxis)).TransformPosition((FVector)PositionVertexBuffer.VertexPosition(Index1) * Mask);
-					const FVector Vert2 = SplineMesh->CalcSliceTransform(USplineMeshComponent::GetAxisValue((FVector)PositionVertexBuffer.VertexPosition(Index2), SplineMesh->ForwardAxis)).TransformPosition((FVector)PositionVertexBuffer.VertexPosition(Index2) * Mask);
+					const FVector Vert0 = SplineMesh->CalcSliceTransform(USplineMeshComponent::GetAxisValue(PositionVertexBuffer.VertexPosition(Index0), SplineMesh->ForwardAxis)).TransformPosition((FVector)PositionVertexBuffer.VertexPosition(Index0) * Mask);
+					const FVector Vert1 = SplineMesh->CalcSliceTransform(USplineMeshComponent::GetAxisValue(PositionVertexBuffer.VertexPosition(Index1), SplineMesh->ForwardAxis)).TransformPosition((FVector)PositionVertexBuffer.VertexPosition(Index1) * Mask);
+					const FVector Vert2 = SplineMesh->CalcSliceTransform(USplineMeshComponent::GetAxisValue(PositionVertexBuffer.VertexPosition(Index2), SplineMesh->ForwardAxis)).TransformPosition((FVector)PositionVertexBuffer.VertexPosition(Index2) * Mask);
 
 					new(PotentialTriangles)FFoliagePaintBucketTriangle(LocalToWorld
 						, Vert0
@@ -2894,12 +2896,12 @@ void FEdModeFoliage::ApplyPaintBucket_Add(AActor* Actor)
 				}
 
 				// This is the total set of instances disregarding parameters like slope, height or layer.
-				float DesiredInstanceCountFloat = Triangle.Area * Settings->Density * UISettings.GetPaintDensity() / (1000.f*1000.f);
+				FVector::FReal DesiredInstanceCountReal = Triangle.Area * Settings->Density * UISettings.GetPaintDensity() / (1000.f*1000.f);
 
 				// Allow a single instance with a random chance, if the brush is smaller than the density
-				int32 DesiredInstanceCount = DesiredInstanceCountFloat > 1.f ? FMath::RoundToInt(DesiredInstanceCountFloat) : FMath::FRand() < DesiredInstanceCountFloat ? 1 : 0;
+				int64 DesiredInstanceCount = DesiredInstanceCountReal > 1.f ? FMath::RoundToInt(DesiredInstanceCountReal) : FMath::FRand() < DesiredInstanceCountReal ? 1 : 0;
 
-				for (int32 Idx = 0; Idx < DesiredInstanceCount; Idx++)
+				for (int64 Idx = 0; Idx < DesiredInstanceCountReal; Idx++)
 				{
 					FVector InstLocation;
 					FColor VertexColor;
@@ -3013,7 +3015,7 @@ bool FEdModeFoliage::GetStaticMeshVertexColorForHit(const UStaticMeshComponent*
 				}
 
 				// find the barycentric coordiantes of the hit location, so we can interpolate the vertex colors
-				FVector BaryCoords = FMath::GetBaryCentric2D(InHitLocation, WorldVert0, WorldVert1, WorldVert2);
+				FVector3f BaryCoords = (FVector3f)FMath::GetBaryCentric2D(InHitLocation, WorldVert0, WorldVert1, WorldVert2);
 
 				FLinearColor InterpColor = BaryCoords.X * Color0 + BaryCoords.Y * Color1 + BaryCoords.Z * Color2;
 
diff --git a/Engine/Source/Editor/FoliageEdit/Private/FoliagePaletteItem.cpp b/Engine/Source/Editor/FoliageEdit/Private/FoliagePaletteItem.cpp
--- a/Engine/Source/Editor/FoliageEdit/Private/FoliagePaletteItem.cpp
+++ b/Engine/Source/Editor/FoliageEdit/Private/FoliagePaletteItem.cpp
@@ -134,7 +134,7 @@ TSharedRef<SToolTip> FFoliagePaletteItemModel::CreateTooltipWidget() const
 {
 	return 
 		SNew(SToolTip)
-		.TextMargin(1)
+		.TextMargin(1.f)
 		.BorderImage(FAppStyle::GetBrush("ContentBrowser.TileViewTooltip.ToolTipBorder"))
 		.Visibility(this, &FFoliagePaletteItemModel::GetTooltipVisibility)
 		[
@@ -265,7 +265,7 @@ TSharedRef<SButton> FFoliagePaletteItemModel::CreateSaveAssetButton(TAttribute<E
 		.ToolTipText(LOCTEXT("SaveButtonToolTip", "Save foliage asset"))
 		.HAlign(HAlign_Right)
 		.VAlign(VAlign_Top)
-		.ContentPadding(0)
+		.ContentPadding(0.f)
 		.Content()
 		[
 			SNew(SImage)
@@ -311,7 +311,7 @@ FText FFoliagePaletteItemModel::GetInstanceCountText(bool bRounded) const
 		};
 
 		int32 NumThousands = 0;
-		float DisplayValue = InstanceCountCurrentLevel;
+		double DisplayValue = InstanceCountCurrentLevel;
 		while (DisplayValue >= 1000.f && NumThousands < EInstanceCountMagnitude::Max)
 		{
 			DisplayValue /= 1000.f;
diff --git a/Engine/Source/Editor/FoliageEdit/Private/ProceduralFoliageEditorLibrary.cpp b/Engine/Source/Editor/FoliageEdit/Private/ProceduralFoliageEditorLibrary.cpp
--- a/Engine/Source/Editor/FoliageEdit/Private/ProceduralFoliageEditorLibrary.cpp
+++ b/Engine/Source/Editor/FoliageEdit/Private/ProceduralFoliageEditorLibrary.cpp
@@ -40,7 +40,7 @@ void UProceduralFoliageEditorLibrary::ClearProceduralFoliageVolumes(const TArray
 
 void UProceduralFoliageEditorLibrary::ClearProceduralFoliageComponents(const TArray<UProceduralFoliageComponent*>& ProceduralFoliageComponents)
 {
-	FScopedSlowTask SlowTask(ProceduralFoliageComponents.Num(), LOCTEXT("ClearProceduralFoliageComponents", "Clearing Procedural Foliage Components"));
+	FScopedSlowTask SlowTask(static_cast<float>(ProceduralFoliageComponents.Num()), LOCTEXT("ClearProceduralFoliageComponents", "Clearing Procedural Foliage Components"));
 	SlowTask.MakeDialogDelayed(0.5f);
 	for (UProceduralFoliageComponent* Component : ProceduralFoliageComponents)
 	{
@@ -54,7 +54,7 @@ void UProceduralFoliageEditorLibrary::ClearProceduralFoliageComponents(const TAr
 
 void UProceduralFoliageEditorLibrary::ResimulateProceduralFoliageComponents(const TArray<UProceduralFoliageComponent*>& ProceduralFoliageComponents)
 {
-	FScopedSlowTask SlowTask(ProceduralFoliageComponents.Num(), LOCTEXT("ResimulateProceduralFoliageComponents", "Resimulating Procedural Foliage Components"));
+	FScopedSlowTask SlowTask(static_cast<float>(ProceduralFoliageComponents.Num()), LOCTEXT("ResimulateProceduralFoliageComponents", "Resimulating Procedural Foliage Components"));
 	SlowTask.MakeDialogDelayed(0.5f);
 	for (UProceduralFoliageComponent* Component : ProceduralFoliageComponents)
 	{
diff --git a/Engine/Source/Editor/FoliageEdit/Private/SFoliageEdit.cpp b/Engine/Source/Editor/FoliageEdit/Private/SFoliageEdit.cpp
--- a/Engine/Source/Editor/FoliageEdit/Private/SFoliageEdit.cpp
+++ b/Engine/Source/Editor/FoliageEdit/Private/SFoliageEdit.cpp
@@ -66,12 +66,12 @@ void SFoliageEdit::Construct(const FArguments& InArgs)
 		SNew(SVerticalBox)
 		+ SVerticalBox::Slot()
 		.AutoHeight()
-		.Padding(0, 0, 0, 5)
+		.Padding(0.f, 0.f, 0.f, 5.f)
 		[
 			SAssignNew(ErrorText, SErrorText)
 		]
 		+ SVerticalBox::Slot()
-		.Padding(0)
+		.Padding(0.f)
 		[
 			SNew(SVerticalBox)
 			.IsEnabled(this, &SFoliageEdit::IsFoliageEditorEnabled)
@@ -254,7 +254,7 @@ void SFoliageEdit::Construct(const FArguments& InArgs)
 							.Visibility(this, &SFoliageEdit::GetVisibility_SingleInstantiationPlacementMode)
 							.IsEnabled(this, &SFoliageEdit::GetIsEnabled_SingleInstantiationPlacementMode)
 							.OnGetMenuContent(this, &SFoliageEdit::GetSingleInstantiationModeMenuContent)
-							.ContentPadding(2)
+							.ContentPadding(2.f)
 							.ToolTipText(LOCTEXT("SingleInstantiationPlacementModeToolTips", "Changes the placement mode when using single instance"))
 							.ButtonContent()
 							[
@@ -282,7 +282,7 @@ void SFoliageEdit::Construct(const FArguments& InArgs)
 							+ SWrapBox::Slot()
 							[
 								SNew(SBox)
-								.MinDesiredWidth(150)
+								.MinDesiredWidth(150.f)
 								[
 									SNew(SCheckBox)
 									.Visibility(this, &SFoliageEdit::GetVisibility_SpawnInCurrentLevelMode)
@@ -331,7 +331,7 @@ void SFoliageEdit::Construct(const FArguments& InArgs)
 							+ SWrapBox::Slot()
 							[
 								SNew(SBox)
-								.MinDesiredWidth(91)
+								.MinDesiredWidth(91.f)
 								.Visibility(this, &SFoliageEdit::GetVisibility_LandscapeFilter)
 								[
 									SNew(SCheckBox)
@@ -350,7 +350,7 @@ void SFoliageEdit::Construct(const FArguments& InArgs)
 							+ SWrapBox::Slot()
 							[
 								SNew(SBox)
-								.MinDesiredWidth(91)
+								.MinDesiredWidth(91.f)
 								[
 									SNew(SCheckBox)
 									.Visibility(this, &SFoliageEdit::GetVisibility_Filters)
@@ -368,7 +368,7 @@ void SFoliageEdit::Construct(const FArguments& InArgs)
 							+ SWrapBox::Slot()
 							[
 								SNew(SBox)
-								.MinDesiredWidth(91)
+								.MinDesiredWidth(91.f)
 								[
 									SNew(SCheckBox)
 									.Visibility(this, &SFoliageEdit::GetVisibility_Filters)
@@ -386,7 +386,7 @@ void SFoliageEdit::Construct(const FArguments& InArgs)
 							+ SWrapBox::Slot()
 							[
 								SNew(SBox)
-								.MinDesiredWidth(91)
+								.MinDesiredWidth(91.f)
 								[
 									SNew(SCheckBox)
 									.Visibility(this, &SFoliageEdit::GetVisibility_Filters)
@@ -404,7 +404,7 @@ void SFoliageEdit::Construct(const FArguments& InArgs)
 							+ SWrapBox::Slot()
 							[
 								SNew(SBox)
-								.MinDesiredWidth(91)
+								.MinDesiredWidth(91.f)
 								[
 									SNew(SCheckBox)
 									.Visibility(this, &SFoliageEdit::GetVisibility_Filters)
diff --git a/Engine/Source/Editor/FoliageEdit/Private/SFoliagePalette.cpp b/Engine/Source/Editor/FoliageEdit/Private/SFoliagePalette.cpp
--- a/Engine/Source/Editor/FoliageEdit/Private/SFoliagePalette.cpp
+++ b/Engine/Source/Editor/FoliageEdit/Private/SFoliagePalette.cpp
@@ -337,7 +337,7 @@ void SFoliagePalette::Construct(const FArguments& InArgs)
 					SNew(SBox)
 					.Visibility(this, &SFoliagePalette::GetDropFoliageHintVisibility)
 					.Padding(FMargin(15, 0))
-					.MinDesiredHeight(30)
+					.MinDesiredHeight(30.f)
 					[
 						SNew(SScaleBox)
 						.Stretch(EStretch::ScaleToFit)
@@ -1124,8 +1124,8 @@ void SFoliagePalette::FillReplaceFoliageTypeSubmenu(FMenuBuilder& MenuBuilder)
 	AssetPickerConfig.bAllowNullSelection = false;
 
 	TSharedRef<SWidget> MenuContent = SNew(SBox)
-		.WidthOverride(384)
-		.HeightOverride(500)
+		.WidthOverride(384.f)
+		.HeightOverride(500.f)
 		[
 			ContentBrowserModule.Get().CreateAssetPicker(AssetPickerConfig)
 		];
diff --git a/Engine/Source/Runtime/Foliage/Private/FoliageComponent.cpp b/Engine/Source/Runtime/Foliage/Private/FoliageComponent.cpp
--- a/Engine/Source/Runtime/Foliage/Private/FoliageComponent.cpp
+++ b/Engine/Source/Runtime/Foliage/Private/FoliageComponent.cpp
@@ -133,7 +133,8 @@ void AInteractiveFoliageActor::SetupCollisionCylinder()
 		const FVector Scale3D = GetStaticMeshComponent()->GetRelativeScale3D();
 		// Set the cylinder's radius based off of the static mesh's bounds radius
 		// CollisionRadius is in world space so apply the actor's scale
-		CapsuleComponent->SetCapsuleSize(MeshBounds.SphereRadius * .7f * FMath::Max(Scale3D.X, Scale3D.Y), MeshBounds.BoxExtent.Z * Scale3D.Z);
+		CapsuleComponent->SetCapsuleSize(static_cast<float>(MeshBounds.SphereRadius * .7f * FMath::Max(Scale3D.X, Scale3D.Y)), static_cast<float>(MeshBounds.BoxExtent.Z * Scale3D.Z));
+
 
 		// Ensure delegate is bound (just once)
 		CapsuleComponent->OnComponentBeginOverlap.RemoveDynamic(this, &AInteractiveFoliageActor::CapsuleTouched);
@@ -218,7 +219,7 @@ void AInteractiveFoliageActor::Tick(float DeltaSeconds)
 		//@todo - derive this height from the static mesh
 		const float IntersectionHeight = 100.0f;
 		// Calculate the rotation angle using Sin(Angle) = Opposite / Hypotenuse
-		const float RotationAngle = -FMath::Asin(FoliagePosition.Size() / IntersectionHeight);
+		const FVector::FReal RotationAngle = -FMath::Asin(FoliagePosition.Size() / IntersectionHeight);
 		// Use a rotation angle perpendicular to the impulse direction and the z axis
 		const FVector NormalizedRotationAxis = FoliagePosition.SizeSquared() > KINDA_SMALL_NUMBER ? 
 			(FoliagePosition ^ FVector(0,0,1)).GetSafeNormal() :
diff --git a/Engine/Source/Runtime/Foliage/Private/InstancedFoliage.cpp b/Engine/Source/Runtime/Foliage/Private/InstancedFoliage.cpp
--- a/Engine/Source/Runtime/Foliage/Private/InstancedFoliage.cpp
+++ b/Engine/Source/Runtime/Foliage/Private/InstancedFoliage.cpp
@@ -210,7 +210,7 @@ namespace FoliageUtil
 			return FoliageType->CullDistance;
 		}
 
-		return FInt32Interval(FoliageType->CullDistance.Min * CullDistanceScale, FoliageType->CullDistance.Max * CullDistanceScale);
+		return FInt32Interval(static_cast<int32>(FoliageType->CullDistance.Min * CullDistanceScale), static_cast<int32>(FoliageType->CullDistance.Max * CullDistanceScale));
 	}
 
 	void UpdateComponentCullDistance(const UFoliageType* FoliageType, UHierarchicalInstancedStaticMeshComponent* Component, float CullDistanceScale)
@@ -262,8 +262,8 @@ struct FFoliagePlacementUtil
 	static int32 GetRandomSeedForPosition(const FVector2D& Position)
 	{
 		// generate a unique random seed for a given position (precision = cm)
-		int32 Xcm = FMath::RoundToInt(Position.X);
-		int32 Ycm = FMath::RoundToInt(Position.Y);
+		int64 Xcm = FMath::RoundToInt(Position.X);
+		int64 Ycm = FMath::RoundToInt(Position.Y);
 		// use the int32 hashing function to avoid patterns by spreading out distribution : 
 		return HashCombine(GetTypeHash(Xcm), GetTypeHash(Ycm));
 	}
@@ -540,7 +540,7 @@ FFoliageDensityFalloff::FFoliageDensityFalloff()
 	FalloffRichCurve->AddKey(1.f, 0.f);
 }
 
-bool FFoliageDensityFalloff::IsInstanceFiltered(const FVector2D& InstancePosition, const FVector2D& Origin, float MaxDistance) const
+bool FFoliageDensityFalloff::IsInstanceFiltered(const FVector2D& InstancePosition, const FVector2D& Origin, FVector::FReal MaxDistance) const
 {
 	float KeepPointProbability = GetDensityFalloffValue(InstancePosition, Origin, MaxDistance);
 	check(KeepPointProbability >= 0.f && KeepPointProbability <= 1.f);
@@ -554,13 +554,13 @@ bool FFoliageDensityFalloff::IsInstanceFiltered(const FVector2D& InstancePositio
 	return false;
 }
 
-float FFoliageDensityFalloff::GetDensityFalloffValue(const FVector2D& Position, const FVector2D& Origin, float MaxDistance) const
+float FFoliageDensityFalloff::GetDensityFalloffValue(const FVector2D& Position, const FVector2D& Origin, FVector::FReal MaxDistance) const
 {
 	float KeepPointProbability = 1.f;
 	if (bUseFalloffCurve)
 	{
-		float Distance = FVector2D::Distance(Position, Origin);
-		float NormalizedDistance = MaxDistance > 0.f ? (Distance / MaxDistance) : 1.f;
+		FVector::FReal Distance = FVector2D::Distance(Position, Origin);
+		float NormalizedDistance = static_cast<float>(MaxDistance > 0.f ? (Distance / MaxDistance) : 1.f);
 		if (NormalizedDistance > 1.f)
 			NormalizedDistance = 1.f;
 		const FRichCurve* FalloffRichCurve = FalloffCurve.GetRichCurveConst();
@@ -806,9 +806,10 @@ bool UFoliageType::IsNotAssetOrBlueprint() const
 	return IsAsset() == false && GetClass()->IsNative();
 }
 
-FVector UFoliageType::GetRandomScale() const
+
+FVector3f UFoliageType::GetRandomScale() const
 {
-	FVector Result(1.0f);
+	FVector3f Result(1.0f);
 	float LockRand = 0.0f;
 
 	switch (Scaling)
@@ -955,7 +956,7 @@ void UFoliageType_Actor::UpdateBounds()
 	FBox LowBound = MeshBounds.GetBox();
 	LowBound.Max.Z = LowBound.Min.Z + (LowBound.Max.Z - LowBound.Min.Z) * 0.1f;
 
-	float MinX = LowBound.Min.X, MaxX = LowBound.Max.X, MinY = LowBound.Min.Y, MaxY = LowBound.Max.Y;
+	FVector::FReal MinX = LowBound.Min.X, MaxX = LowBound.Max.X, MinY = LowBound.Min.Y, MaxY = LowBound.Max.Y;
 	LowBoundOriginRadius = FVector::ZeroVector;
 
 	// TODO: Get more precise lower bound from multiple possible meshes in Actor
@@ -981,7 +982,7 @@ float UFoliageType::GetScaleForAge(const float Age) const
 
 float UFoliageType::GetInitAge(FRandomStream& RandomStream) const
 {
-	return RandomStream.FRandRange(0, MaxInitialAge);
+	return MaxInitialAge * RandomStream.GetFraction();
 }
 
 float UFoliageType::GetNextAge(const float CurrentAge, const int32 InNumSteps) const
@@ -1640,7 +1641,7 @@ void FFoliageStaticMesh::UpdateComponentSettings(const UFoliageType_InstancedSta
 		}
 		if (Component->VirtualTextureCullMips != FoliageType->VirtualTextureCullMips)
 		{
-			Component->VirtualTextureCullMips = FoliageType->VirtualTextureCullMips;
+			Component->VirtualTextureCullMips = static_cast<int8>(FoliageType->VirtualTextureCullMips);
 			bNeedsMarkRenderStateDirty = true;
 		}
 		if (Component->TranslucencySortPriority != FoliageType->TranslucencySortPriority)
@@ -2681,13 +2682,13 @@ void FFoliageInfo::GetInstanceAtLocation(const FVector& Location, int32& OutInst
 {
 	auto TempInstances = InstanceHash->GetInstancesOverlappingBox(FBox::BuildAABB(Location, FVector(KINDA_SMALL_NUMBER)));
 
-	float ShortestDistance = MAX_FLT;
+	FVector::FReal ShortestDistance = MAX_dbl;
 	OutInstance = -1;
 
 	for (int32 Idx : TempInstances)
 	{
 		FVector InstanceLocation = Instances[Idx].Location;
-		float DistanceSquared = FVector::DistSquared(InstanceLocation, Location);
+		FVector::FReal DistanceSquared = FVector::DistSquared(InstanceLocation, Location);
 		if (DistanceSquared < ShortestDistance)
 		{
 			ShortestDistance = DistanceSquared;
@@ -5343,7 +5344,7 @@ bool AInstancedFoliageActor::FoliageTrace(const UWorld* InWorld, FHitResult& Out
 					{
 						FBox ActorVolumeBounds = Brush->Bounds.GetBox();
 						FVector2D ActorVolumeLocation = FVector2D(ActorVolumeBounds.GetCenter());
-						const float ActorVolumeMaxExtent = FVector2D(ActorVolumeBounds.GetExtent()).GetMax();
+						const FVector::FReal ActorVolumeMaxExtent = FVector2D(ActorVolumeBounds.GetExtent()).GetMax();
 
 						const FVector2D Origin(ProceduralFoliageBlockingVolume->GetActorTransform().GetLocation());
 						if (ProceduralFoliageBlockingVolume->DensityFalloff.IsInstanceFiltered(FVector2D(Hit.ImpactPoint), ActorVolumeLocation, ActorVolumeMaxExtent))
@@ -5463,8 +5464,8 @@ bool AInstancedFoliageActor::FoliageTrace(const UWorld* InWorld, FHitResult& Out
 			bool bSingleComponent = DesiredInstance.FoliageType->AverageNormalSingleComponent;
 			for (int32 i = 0; i < DesiredInstance.FoliageType->AverageNormalSampleCount; ++i)
 			{
-				const float Angle = LocalRandomStream.FRandRange(0, PI * 2.f);
-				const float SqrtRadius = FMath::Sqrt(LocalRandomStream.FRand()) * DesiredInstance.FoliageType->LowBoundOriginRadius.Z;
+				const FVector::FReal Angle = LocalRandomStream.FRandRange(0, PI * 2.f);
+				const FVector::FReal SqrtRadius = FMath::Sqrt(LocalRandomStream.FRand()) * DesiredInstance.FoliageType->LowBoundOriginRadius.Z;
 				FVector Offset(SqrtRadius * FMath::Cos(Angle), SqrtRadius* FMath::Sin(Angle), 0.f);
 				NormalHits.Reset();
 				if (InWorld->LineTraceMultiByObjectType(NormalHits, StartTrace + Offset, DesiredInstance.EndTrace + Offset, FCollisionObjectQueryParams(ECC_WorldStatic), QueryParams))
@@ -5573,7 +5574,7 @@ bool FPotentialInstance::PlaceInstance(const UWorld* InWorld, const UFoliageType
 {
 	if (DesiredInstance.PlacementMode != EFoliagePlacementMode::Procedural)
 	{
-		Inst.DrawScale3D = (FVector3f)Settings->GetRandomScale();
+		Inst.DrawScale3D = Settings->GetRandomScale();
 		Inst.ZOffset = Settings->ZOffset.Interpolate(FMath::FRand());
 	}
 	else
@@ -5903,7 +5904,7 @@ void UFoliageInstancedStaticMeshComponent::ReceiveComponentDamage(float DamageAm
 			if (Instances.Num())
 			{
 				FVector LocalOrigin = GetComponentToWorld().Inverse().TransformPosition(RadialDamageEvent->Origin);
-				float Scale = GetComponentScale().X; // assume component (not instances) is uniformly scaled
+				FVector::FReal Scale = GetComponentScale().X; // assume component (not instances) is uniformly scaled
 
 				TArray<float> Damages;
 				Damages.Empty(Instances.Num());
@@ -5911,8 +5912,9 @@ void UFoliageInstancedStaticMeshComponent::ReceiveComponentDamage(float DamageAm
 				for (int32 InstanceIndex : Instances)
 				{
 					// Find distance in local space and then scale; quicker than transforming each instance to world space.
-					float DistanceFromOrigin = (PerInstanceSMData[InstanceIndex].Transform.GetOrigin() - LocalOrigin).Size() * Scale;
+					float DistanceFromOrigin = static_cast<float>((PerInstanceSMData[InstanceIndex].Transform.GetOrigin() - LocalOrigin).Size() * Scale);
 					Damages.Add(RadialDamageEvent->Params.GetDamageScale(DistanceFromOrigin));
+
 				}
 
 				OnInstanceTakeRadialDamage.Broadcast(Instances, Damages, EventInstigator, RadialDamageEvent->Origin, MaxRadius, DamageTypeCDO, DamageCauser);
diff --git a/Engine/Source/Runtime/Foliage/Private/ProceduralFoliageComponent.cpp b/Engine/Source/Runtime/Foliage/Private/ProceduralFoliageComponent.cpp
--- a/Engine/Source/Runtime/Foliage/Private/ProceduralFoliageComponent.cpp
+++ b/Engine/Source/Runtime/Foliage/Private/ProceduralFoliageComponent.cpp
@@ -99,13 +99,13 @@ void UProceduralFoliageComponent::GetTileLayout(const FGenerateProceduralContent
 	{
 		// Determine the bottom-left-most tile that contains the min position (when accounting for overlap)
 		const FVector MinPosition = InParams.Bounds.Min + InParams.TileOverlap;
-		OutTileLayout.BottomLeftX = FMath::FloorToInt(MinPosition.X / InParams.FoliageSpawner->TileSize);
-		OutTileLayout.BottomLeftY = FMath::FloorToInt(MinPosition.Y / InParams.FoliageSpawner->TileSize);
+		OutTileLayout.BottomLeftX = FMath::FloorToInt(FloatCastChecked<float>(MinPosition.X / InParams.FoliageSpawner->TileSize, UE::LWC::DefaultFloatPrecision));
+		OutTileLayout.BottomLeftY = FMath::FloorToInt(FloatCastChecked<float>(MinPosition.Y / InParams.FoliageSpawner->TileSize, UE::LWC::DefaultFloatPrecision));
 
 		// Determine the total number of tiles along each active axis
 		const FVector MaxPosition = InParams.Bounds.Max - InParams.TileOverlap;
-		const int32 MaxXIdx = FMath::FloorToInt(MaxPosition.X / InParams.FoliageSpawner->TileSize);
-		const int32 MaxYIdx = FMath::FloorToInt(MaxPosition.Y / InParams.FoliageSpawner->TileSize);
+		const int32 MaxXIdx = FMath::FloorToInt(FloatCastChecked<float>(MaxPosition.X / InParams.FoliageSpawner->TileSize, UE::LWC::DefaultFloatPrecision));
+		const int32 MaxYIdx = FMath::FloorToInt(FloatCastChecked<float>(MaxPosition.Y / InParams.FoliageSpawner->TileSize, UE::LWC::DefaultFloatPrecision));
 
 		OutTileLayout.NumTilesX = (MaxXIdx - OutTileLayout.BottomLeftX) + 1;
 		OutTileLayout.NumTilesY = (MaxYIdx - OutTileLayout.BottomLeftY) + 1;
@@ -199,7 +199,7 @@ bool UProceduralFoliageComponent::GenerateProceduralContent(const FGenerateProce
 		const float TileSize = InParams.FoliageSpawner->TileSize;
 		const FVector WorldPosition = GetWorldPosition(InParams);
 		FVector2D GenerateLocation = FVector2D(InParams.Bounds.GetCenter());
-		const float GenerateMaxExtent = FVector2D(InParams.Bounds.GetExtent()).GetMax();
+		const FVector::FReal GenerateMaxExtent = FVector2D(InParams.Bounds.GetExtent()).GetMax();
 		FTileLayout TileLayout;
 		GetTileLayout(InParams, TileLayout);
 
diff --git a/Engine/Source/Runtime/Foliage/Private/ProceduralFoliageSpawner.cpp b/Engine/Source/Runtime/Foliage/Private/ProceduralFoliageSpawner.cpp
--- a/Engine/Source/Runtime/Foliage/Private/ProceduralFoliageSpawner.cpp
+++ b/Engine/Source/Runtime/Foliage/Private/ProceduralFoliageSpawner.cpp
@@ -77,12 +77,12 @@ const UProceduralFoliageTile* UProceduralFoliageSpawner::GetRandomTile(int32 X,
 		FRandomStream HashStream;	
 		
 		HashStream.Initialize(X);
-		const double XRand = HashStream.FRand();
+		const float XRand = HashStream.FRand();
 		
 		HashStream.Initialize(Y);
-		const double YRand = HashStream.FRand();
+		const float YRand = HashStream.FRand();
 		
-		const int32 RandomNumber = (RAND_MAX * XRand / (YRand + 0.01));
+		const int32 RandomNumber = static_cast<int32>(RAND_MAX * XRand / (YRand + 0.01));
 		const int32 Idx = FMath::Clamp(RandomNumber % PrecomputedTiles.Num(), 0, PrecomputedTiles.Num() - 1);
 		return PrecomputedTiles[Idx].Get();
 	}
@@ -153,7 +153,7 @@ void UProceduralFoliageSpawner::Simulate(int32 NumSteps)
 
 int32 UProceduralFoliageSpawner::GetRandomNumber()
 {
-	return RandomStream.FRand() * float(RAND_MAX);
+	return static_cast<int32>(RandomStream.FRand() * float(RAND_MAX));
 }
 
 #undef LOCTEXT_NAMESPACE
diff --git a/Engine/Source/Runtime/Foliage/Private/ProceduralFoliageTile.cpp b/Engine/Source/Runtime/Foliage/Private/ProceduralFoliageTile.cpp
--- a/Engine/Source/Runtime/Foliage/Private/ProceduralFoliageTile.cpp
+++ b/Engine/Source/Runtime/Foliage/Private/ProceduralFoliageTile.cpp
@@ -101,7 +101,7 @@ FProceduralFoliageInstance* UProceduralFoliageTile::NewSeed(const FVector& Locat
 float GetSeedMinDistance(const FProceduralFoliageInstance* Instance, const float NewInstanceAge, const int32 SimulationStep)
 {
 	const UFoliageType* Type = Instance->Type;
-	const int32 StepsLeft = Type->MaxAge - SimulationStep;
+	const int32 StepsLeft = static_cast<int32>(Type->MaxAge - SimulationStep);
 	const float InstanceMaxAge = Type->GetNextAge(Instance->Age, StepsLeft);
 	const float NewInstanceMaxAge = Type->GetNextAge(NewInstanceAge, StepsLeft);
 
@@ -291,13 +291,13 @@ void UProceduralFoliageTile::AddRandomSeeds(TArray<FProceduralFoliageInstance*>&
 			const float Scale = Type->GetScaleForAge(NewAge);
 
 			FRandomStream& TypeRandomStream = RandomStreamPerType.FindChecked(Type);
-			float InitX = 0.f;
-			float InitY = 0.f;
+			FVector::FReal InitX = 0.f;
+			FVector::FReal InitY = 0.f;
 			float NeededRadius = 0.f;
 
 			if (bSimulateOnlyInShade && LastShadeCastingIndex >= 0)
 			{
-				const int32 InstanceSpawnerIdx = TypeRandomStream.FRandRange(0, LastShadeCastingIndex);
+				const int32 InstanceSpawnerIdx = TypeRandomStream.RandRange(0, LastShadeCastingIndex);
 				const FProceduralFoliageInstance& Spawner = InstancesArray[InstanceSpawnerIdx];
 				InitX = Spawner.Location.X;
 				InitY = Spawner.Location.Y;
@@ -310,13 +310,13 @@ void UProceduralFoliageTile::AddRandomSeeds(TArray<FProceduralFoliageInstance*>&
 				NeededRadius = MaxShadeRadii.FindRef(Type->DistributionSeed);
 			}
 
-			const float Rad = RandomStream.FRandRange(0, PI*2.f);
+			const FVector::FReal Rad = RandomStream.FRandRange(0, PI*2.f);
 			
 			
 			const FVector GlobalOffset = (RandomStream.FRandRange(0, Type->MaxInitialSeedOffset) + NeededRadius) * FVector(FMath::Cos(Rad), FMath::Sin(Rad), 0.f);
 
-			const float X = InitX + GlobalOffset.X;
-			const float Y = InitY + GlobalOffset.Y;
+			const FVector::FReal X = InitX + GlobalOffset.X;
+			const FVector::FReal Y = InitY + GlobalOffset.Y;
 
 			if (FProceduralFoliageInstance* NewInst = NewSeed(FVector(X, Y, 0.f), Scale, Type, NewAge))
 			{
@@ -467,7 +467,7 @@ void UProceduralFoliageTile::BeginDestroy()
 	RemoveInstances();
 }
 
-void UProceduralFoliageTile::ExtractDesiredInstances(TArray<FDesiredFoliageInstance>& OutInstances, const FTransform& WorldTM, const FVector2D& ActorVolumeLocation, float ActorVolumeMaxExtent, const FGuid& ProceduralGuid, const float HalfHeight, const FBodyInstance* VolumeBodyInstance, bool bEmptyTileInfo)
+void UProceduralFoliageTile::ExtractDesiredInstances(TArray<FDesiredFoliageInstance>& OutInstances, const FTransform& WorldTM, const FVector2D& ActorVolumeLocation, FVector::FReal ActorVolumeMaxExtent, const FGuid& ProceduralGuid, const FVector::FReal HalfHeight, const FBodyInstance* VolumeBodyInstance, bool bEmptyTileInfo)
 {
 	InstancesToArray();
 
diff --git a/Engine/Source/Runtime/Foliage/Public/FoliageType.h b/Engine/Source/Runtime/Foliage/Public/FoliageType.h
--- a/Engine/Source/Runtime/Foliage/Public/FoliageType.h
+++ b/Engine/Source/Runtime/Foliage/Public/FoliageType.h
@@ -97,8 +97,8 @@ struct FFoliageDensityFalloff
 	UPROPERTY(Category = Procedural, EditAnywhere, meta = (Subcategory = "Density", XAxisName = "Normalized Distance", YAxisName = "Density Factor"))
 	FRuntimeFloatCurve FalloffCurve;
 
-	FOLIAGE_API bool IsInstanceFiltered(const FVector2D& Position, const FVector2D& Origin, float MaxDistance) const;
-	FOLIAGE_API float GetDensityFalloffValue(const FVector2D& Position, const FVector2D& Origin, float MaxDistance) const;
+	FOLIAGE_API bool IsInstanceFiltered(const FVector2D& Position, const FVector2D& Origin, FVector::FReal MaxDistance) const;
+	FOLIAGE_API float GetDensityFalloffValue(const FVector2D& Position, const FVector2D& Origin, FVector::FReal MaxDistance) const;
 };
 
 UCLASS(abstract, hidecategories = Object, editinlinenew, MinimalAPI, BlueprintType, Blueprintable)
@@ -113,7 +113,7 @@ class UFoliageType : public UObject
 	virtual void PostLoad() override;
 
 	virtual bool IsNotAssetOrBlueprint() const;
-	FOLIAGE_API FVector GetRandomScale() const;
+	FOLIAGE_API FVector3f GetRandomScale() const;
 	
 #if WITH_EDITOR
 
diff --git a/Engine/Source/Runtime/Foliage/Public/InstancedFoliage.h b/Engine/Source/Runtime/Foliage/Public/InstancedFoliage.h
--- a/Engine/Source/Runtime/Foliage/Public/InstancedFoliage.h
+++ b/Engine/Source/Runtime/Foliage/Public/InstancedFoliage.h
@@ -119,7 +119,7 @@ struct FFoliageInstance : public FFoliageInstancePlacementInfo
 		// limit the maximum pitch angle if it's > 0.
 		if (AlignMaxAngle > 0.f)
 		{
-			int32 MaxPitch = AlignMaxAngle;
+			int32 MaxPitch = static_cast<int32>(AlignMaxAngle);
 			if (AlignRotation.Pitch > MaxPitch)
 			{
 				AlignRotation.Pitch = MaxPitch;
diff --git a/Engine/Source/Runtime/Foliage/Public/ProceduralFoliageComponent.h b/Engine/Source/Runtime/Foliage/Public/ProceduralFoliageComponent.h
--- a/Engine/Source/Runtime/Foliage/Public/ProceduralFoliageComponent.h
+++ b/Engine/Source/Runtime/Foliage/Public/ProceduralFoliageComponent.h
@@ -35,7 +35,7 @@ struct FTileLayout
 	int32 NumTilesY;
 	
 
-	float HalfHeight;
+	FVector::FReal HalfHeight;
 };
 
 UCLASS(BlueprintType)
diff --git a/Engine/Source/Runtime/Foliage/Public/ProceduralFoliageTile.h b/Engine/Source/Runtime/Foliage/Public/ProceduralFoliageTile.h
--- a/Engine/Source/Runtime/Foliage/Public/ProceduralFoliageTile.h
+++ b/Engine/Source/Runtime/Foliage/Public/ProceduralFoliageTile.h
@@ -42,7 +42,7 @@ class FOLIAGE_API UProceduralFoliageTile : public UObject
 	/**
 	 *	Extracts information on each instance that should be spawned based on the simulation.
 	 */
-	void ExtractDesiredInstances(TArray<FDesiredFoliageInstance>& OutDesiredInstances, const FTransform& WorldTM, const FVector2D& ActorVolumeLocation, float ActorVolumeMaxExtent, const FGuid& ProceduralGuid, const float HalfHeight, const FBodyInstance* VolumeBodyInstance, bool bEmptyTileInfo = true);
+	void ExtractDesiredInstances(TArray<FDesiredFoliageInstance>& OutDesiredInstances, const FTransform& WorldTM, const FVector2D& ActorVolumeLocation, FVector::FReal ActorVolumeMaxExtent, const FGuid& ProceduralGuid, const FVector::FReal HalfHeight, const FBodyInstance* VolumeBodyInstance, bool bEmptyTileInfo = true);
 	
 	/**
 	 *	Copies instances within this tile to another
