Fix Build: Android

#rnx

diff --git a/Engine/Source/Runtime/Foliage/Private/ProceduralFoliageSpawner.cpp b/Engine/Source/Runtime/Foliage/Private/ProceduralFoliageSpawner.cpp
--- a/Engine/Source/Runtime/Foliage/Private/ProceduralFoliageSpawner.cpp
+++ b/Engine/Source/Runtime/Foliage/Private/ProceduralFoliageSpawner.cpp
@@ -82,7 +82,7 @@ const UProceduralFoliageTile* UProceduralFoliageSpawner::GetRandomTile(int32 X,
 		HashStream.Initialize(Y);
 		const float YRand = HashStream.FRand();
 		
-		const int32 RandomNumber = static_cast<int32>(RAND_MAX * XRand / (YRand + 0.01));
+		const int32 RandomNumber = static_cast<int32>(static_cast<float>(RAND_MAX) * XRand / (YRand + 0.01));
 		const int32 Idx = FMath::Clamp(RandomNumber % PrecomputedTiles.Num(), 0, PrecomputedTiles.Num() - 1);
 		return PrecomputedTiles[Idx].Get();
 	}
