ProceduralFoliageBlockingVolume: Initialize bIsEditorOnlyActor = true by default

diff --git a/Engine/Source/Runtime/Foliage/Private/ProceduralFoliageBlockingVolume.cpp b/Engine/Source/Runtime/Foliage/Private/ProceduralFoliageBlockingVolume.cpp
--- a/Engine/Source/Runtime/Foliage/Private/ProceduralFoliageBlockingVolume.cpp
+++ b/Engine/Source/Runtime/Foliage/Private/ProceduralFoliageBlockingVolume.cpp
@@ -15,5 +15,7 @@ AProceduralFoliageBlockingVolume::AProceduralFoliageBlockingVolume(const FObject
 		MyBrushComponent->SetCollisionObjectType(ECC_WorldStatic);
 		MyBrushComponent->SetCollisionResponseToAllChannels(ECR_Ignore);
 	}
+
+	bIsEditorOnlyActor = true;
 }
 
