Fixed Foliage Replace With option where target foliage type is already used : InstanceFoliageActor components weren't cleaned up before merging foliage instances.

diff --git a/Engine/Source/Runtime/Foliage/Private/InstancedFoliage.cpp b/Engine/Source/Runtime/Foliage/Private/InstancedFoliage.cpp
--- a/Engine/Source/Runtime/Foliage/Private/InstancedFoliage.cpp
+++ b/Engine/Source/Runtime/Foliage/Private/InstancedFoliage.cpp
@@ -2623,6 +2623,10 @@ void FFoliageInfo::RecomputeHash()
 void FFoliageInfo::ReallocateClusters(UFoliageType* InSettings)
 {
 	// In case Foliage Type Changed recreate implementation
+	if (Implementation.IsValid())
+	{
+		Implementation->Uninitialize();
+	}
 	Implementation.Reset();
 	CreateImplementation(InSettings);
 	
