Foliage optim: use a static name instead of looking it up for every call to FFoliageHelper::IsOwnedByFoliage.
This happens when we rename actors during cooks for partitioned world, for every actor.

diff --git a/Engine/Source/Runtime/Foliage/Public/FoliageHelper.h b/Engine/Source/Runtime/Foliage/Public/FoliageHelper.h
--- a/Engine/Source/Runtime/Foliage/Public/FoliageHelper.h
+++ b/Engine/Source/Runtime/Foliage/Public/FoliageHelper.h
@@ -12,22 +12,31 @@ class FFoliageHelper
 {
 public:
 #if WITH_EDITOR
-	static void SetIsOwnedByFoliage(AActor* InActor, bool bOwned = true) 
-	{ 
+	static inline FName GetFoliageActorInstanceTag()
+	{
 		static FName NAME_FoliageActorInstanceTag(TEXT("FoliageActorInstance"));
+		return NAME_FoliageActorInstanceTag;
+	}
+
+	static inline void SetIsOwnedByFoliage(AActor* InActor, bool bOwned = true) 
+	{ 
 		if (InActor) 
 		{ 
 			FSetActorHiddenInSceneOutliner SetHidden(InActor, bOwned);
 			if (bOwned)
 			{
-				InActor->Tags.AddUnique(NAME_FoliageActorInstanceTag);
+				InActor->Tags.AddUnique(GetFoliageActorInstanceTag());
 			}
 			else
 			{
-				InActor->Tags.Remove(NAME_FoliageActorInstanceTag);
+				InActor->Tags.Remove(GetFoliageActorInstanceTag());
 			}
 		} 
 	}
-	static bool IsOwnedByFoliage(const AActor* InActor) { return InActor != nullptr && InActor->ActorHasTag(TEXT("FoliageActorInstance")); }
+
+	static inline bool IsOwnedByFoliage(const AActor* InActor)
+	{
+		return InActor && InActor->ActorHasTag(GetFoliageActorInstanceTag());
+	}
 #endif
 };
\ No newline at end of file
