Fixed crash in PIE/Cook in world partition maps using old/unsupported Foliage Type Actor.

diff --git a/Engine/Source/Runtime/Foliage/Private/InstancedFoliage.cpp b/Engine/Source/Runtime/Foliage/Private/InstancedFoliage.cpp
--- a/Engine/Source/Runtime/Foliage/Private/InstancedFoliage.cpp
+++ b/Engine/Source/Runtime/Foliage/Private/InstancedFoliage.cpp
@@ -5048,7 +5048,8 @@ void AInstancedFoliageActor::OnLevelActorOuterChanged(AActor* InActor, UObject*
 		return;
 	}
 
-	if (OldLevel)
+	UActorPartitionSubsystem* ActorPartitionSubsystem = OldLevel ? UWorld::GetSubsystem<UActorPartitionSubsystem>(OldLevel->GetWorld()) : nullptr;
+	if (ActorPartitionSubsystem && ActorPartitionSubsystem->IsLevelPartition())
 	{
 		AInstancedFoliageActor* OldIFA = AInstancedFoliageActor::GetInstancedFoliageActorForLevel(OldLevel, false);
 		check(OldIFA);
