Do not reset the FFoliageInfo() when running WITH_EDITOR in -game.  The implementation handle is needed for transactional events sent by Multi-user

diff --git a/Engine/Source/Runtime/Foliage/Private/InstancedFoliage.cpp b/Engine/Source/Runtime/Foliage/Private/InstancedFoliage.cpp
--- a/Engine/Source/Runtime/Foliage/Private/InstancedFoliage.cpp
+++ b/Engine/Source/Runtime/Foliage/Private/InstancedFoliage.cpp
@@ -4855,8 +4855,11 @@ void AInstancedFoliageActor::PostLoad()
 					}
 				}
 			}
-				
+
+#if !WITH_EDITOR
+			// If we are running in -game we still need foliage info for any remote transactional events.
 			Pair.Value = FFoliageInfo();
+#endif
 		}
 
 		if (bHasISMFoliage)
