FoliageType: Fix failed check when changing 'bStaticMeshOnly' property

diff --git a/Engine/Source/Runtime/Foliage/Private/FoliageActor.cpp b/Engine/Source/Runtime/Foliage/Private/FoliageActor.cpp
--- a/Engine/Source/Runtime/Foliage/Private/FoliageActor.cpp
+++ b/Engine/Source/Runtime/Foliage/Private/FoliageActor.cpp
@@ -285,20 +285,19 @@ void FFoliageActor::PostMoveInstances(TArrayView<const int32> InInstancesMoved,
 	}
 }
 
-void FFoliageActor::NotifyFoliageTypeChanged(UFoliageType* FoliageType, bool bSourceChanged)
+bool FFoliageActor::NotifyFoliageTypeChanged(UFoliageType* FoliageType, bool bSourceChanged)
 {
 	if (!IsInitialized())
 	{
-		return;
+		return false;
 	}
 
 	if (UFoliageType_Actor* FoliageTypeActor = Cast<UFoliageType_Actor>(FoliageType))
 	{
-		// Implementation should change
 		if (FoliageTypeActor->bStaticMeshOnly)
 		{
-			Uninitialize();
-			return;
+			// requires implementation change
+			return true;
 		}
 
 		AInstancedFoliageActor* IFA = GetIFA();
@@ -317,6 +316,8 @@ void FFoliageActor::NotifyFoliageTypeChanged(UFoliageType* FoliageType, bool bSo
 		Reapply(FoliageType);
 		ApplySelection(true, Info->SelectedIndices);
 	}
+
+	return false;
 }
 
 void FFoliageActor::Reapply(const UFoliageType* FoliageType)
diff --git a/Engine/Source/Runtime/Foliage/Private/FoliageActor.h b/Engine/Source/Runtime/Foliage/Private/FoliageActor.h
--- a/Engine/Source/Runtime/Foliage/Private/FoliageActor.h
+++ b/Engine/Source/Runtime/Foliage/Private/FoliageActor.h
@@ -57,7 +57,7 @@ struct FFoliageActor : public FFoliageImpl
 	virtual void PostEditUndo(FFoliageInfo* InInfo, UFoliageType* FoliageType) override;
 	virtual void PreMoveInstances(TArrayView<const int32> InInstancesMoved) override;
 	virtual void PostMoveInstances(TArrayView<const int32> InInstancesMoved, bool bFinished) override;
-	virtual void NotifyFoliageTypeChanged(UFoliageType* FoliageType, bool bSourceChanged) override;
+	virtual bool NotifyFoliageTypeChanged(UFoliageType* FoliageType, bool bSourceChanged) override;
 	virtual void Reapply(const UFoliageType* FoliageType) override;
 	AActor* Spawn(const FFoliageInstance& Instance);
 	TArray<AActor*> GetActorsFromSelectedIndices(const TSet<int32>& SelectedIndices) const;
diff --git a/Engine/Source/Runtime/Foliage/Private/FoliageISMActor.cpp b/Engine/Source/Runtime/Foliage/Private/FoliageISMActor.cpp
--- a/Engine/Source/Runtime/Foliage/Private/FoliageISMActor.cpp
+++ b/Engine/Source/Runtime/Foliage/Private/FoliageISMActor.cpp
@@ -331,20 +331,19 @@ void FFoliageISMActor::NotifyFoliageTypeWillChange(UFoliageType* FoliageType)
 	UnregisterDelegates();
 }
 
-void FFoliageISMActor::NotifyFoliageTypeChanged(UFoliageType* FoliageType, bool bSourceChanged)
+bool FFoliageISMActor::NotifyFoliageTypeChanged(UFoliageType* FoliageType, bool bSourceChanged)
 {
 	if (!IsInitialized())
 	{
-		return;
+		return false;
 	}
 		
 	if (UFoliageType_Actor* InFoliageTypeActor = Cast<UFoliageType_Actor>(FoliageType))
 	{
-		// Implementation should change
 		if (!InFoliageTypeActor->bStaticMeshOnly)
 		{
-			Uninitialize();
-			return;
+			// requires implementation change
+			return true;
 		}
 	}
 
@@ -377,6 +376,8 @@ void FFoliageISMActor::NotifyFoliageTypeChanged(UFoliageType* FoliageType, bool
 	{
 		RegisterDelegates();
 	}
+
+	return false;
 }
 
 void FFoliageISMActor::SelectAllInstances(bool bSelect)
diff --git a/Engine/Source/Runtime/Foliage/Private/FoliageISMActor.h b/Engine/Source/Runtime/Foliage/Private/FoliageISMActor.h
--- a/Engine/Source/Runtime/Foliage/Private/FoliageISMActor.h
+++ b/Engine/Source/Runtime/Foliage/Private/FoliageISMActor.h
@@ -68,7 +68,7 @@ struct FFoliageISMActor : public FFoliageImpl, public IISMPartitionInstanceManag
 	virtual void PreEditUndo(UFoliageType* FoliageType) override;
 	virtual void PostEditUndo(FFoliageInfo* InInfo, UFoliageType* FoliageType) override;
 	virtual void NotifyFoliageTypeWillChange(UFoliageType* FoliageType) override;
-	virtual void NotifyFoliageTypeChanged(UFoliageType* FoliageType, bool bSourceChanged) override;
+	virtual bool NotifyFoliageTypeChanged(UFoliageType* FoliageType, bool bSourceChanged) override;
 
 private:
 	void RegisterDelegates();
diff --git a/Engine/Source/Runtime/Foliage/Private/InstancedFoliage.cpp b/Engine/Source/Runtime/Foliage/Private/InstancedFoliage.cpp
--- a/Engine/Source/Runtime/Foliage/Private/InstancedFoliage.cpp
+++ b/Engine/Source/Runtime/Foliage/Private/InstancedFoliage.cpp
@@ -165,7 +165,7 @@ struct FFoliageStaticMesh : public FFoliageImpl, public ISMInstanceManager
 	virtual void PreEditUndo(UFoliageType* FoliageType) override;
 	virtual void PostEditUndo(FFoliageInfo* InInfo, UFoliageType* FoliageType) override;
 	virtual void NotifyFoliageTypeWillChange(UFoliageType* FoliageType) override;
-	virtual void NotifyFoliageTypeChanged(UFoliageType* FoliageType, bool bSourceChanged) override;
+	virtual bool NotifyFoliageTypeChanged(UFoliageType* FoliageType, bool bSourceChanged) override;
 	virtual void EnterEditMode() override;
 	virtual void ExitEditMode() override;
 
@@ -1435,7 +1435,7 @@ void FFoliageStaticMesh::NotifyFoliageTypeWillChange(UFoliageType* FoliageType)
 	}
 }
 
-void FFoliageStaticMesh::NotifyFoliageTypeChanged(UFoliageType* FoliageType, bool bSourceChanged)
+bool FFoliageStaticMesh::NotifyFoliageTypeChanged(UFoliageType* FoliageType, bool bSourceChanged)
 {
 	UFoliageType_InstancedStaticMesh* FoliageType_InstancedStaticMesh = Cast<UFoliageType_InstancedStaticMesh>(FoliageType);
 	check(FoliageType_InstancedStaticMesh);
@@ -1453,6 +1453,8 @@ void FFoliageStaticMesh::NotifyFoliageTypeChanged(UFoliageType* FoliageType, boo
 			Component->BuildTreeIfOutdated(true, false);
 		}
 	}
+
+	return false;
 }
 
 void FFoliageStaticMesh::EnterEditMode()
@@ -2157,11 +2159,7 @@ void FFoliageInfo::NotifyFoliageTypeWillChange(UFoliageType* FoliageType)
 void FFoliageInfo::NotifyFoliageTypeChanged(UFoliageType* FoliageType, bool bSourceChanged)
 {
 	FoliageTypeUpdateGuid = FoliageType->UpdateGuid;
-	// Handle Implementation being uninitialized by FoliageType change
-	bool bWasInitialized = Implementation->IsInitialized();
-	
-	Implementation->NotifyFoliageTypeChanged(FoliageType, bSourceChanged);
-	if (bWasInitialized && !Implementation->IsInitialized())
+	if (Implementation->NotifyFoliageTypeChanged(FoliageType, bSourceChanged))
 	{
 		ReallocateClusters(FoliageType);
 	}
diff --git a/Engine/Source/Runtime/Foliage/Public/InstancedFoliage.h b/Engine/Source/Runtime/Foliage/Public/InstancedFoliage.h
--- a/Engine/Source/Runtime/Foliage/Public/InstancedFoliage.h
+++ b/Engine/Source/Runtime/Foliage/Public/InstancedFoliage.h
@@ -247,7 +247,9 @@ struct FFoliageImpl
 	virtual void PreEditUndo(UFoliageType* FoliageType) {}
 	virtual void PostEditUndo(FFoliageInfo* InInfo, UFoliageType* FoliageType) { Info = InInfo; }
 	virtual void NotifyFoliageTypeWillChange(UFoliageType* FoliageType) {}
-	virtual void NotifyFoliageTypeChanged(UFoliageType* FoliageType, bool bSourceChanged) = 0;
+
+	// Return true if implementation needs to change 
+	virtual bool NotifyFoliageTypeChanged(UFoliageType* FoliageType, bool bSourceChanged) = 0;
 	virtual void EnterEditMode() {}
 	virtual void ExitEditMode() {}
 	virtual bool ShouldAttachToBaseComponent() const { return true; }
