Adding properties to UProceduralFoliageSpawner to allow for more customization of how foliage is spawned on materials.  Derived implementations of UProceduralFoliageComponent can use this data to filter where foliage is spawned.
Part 1/2 Split Submit

#tests Used grass preview, rebuilt grass using commandlet, and inspected built grass in the editor.

diff --git a/Engine/Source/Runtime/Foliage/Public/ProceduralFoliageSpawner.h b/Engine/Source/Runtime/Foliage/Public/ProceduralFoliageSpawner.h
--- a/Engine/Source/Runtime/Foliage/Public/ProceduralFoliageSpawner.h
+++ b/Engine/Source/Runtime/Foliage/Public/ProceduralFoliageSpawner.h
@@ -41,6 +41,17 @@ private:
 	UPROPERTY(Category = ProceduralFoliageSimulation, EditAnywhere)
 	TArray<FFoliageTypeObject> FoliageTypes;
 
+	/**
+	 * If checked, will override the default behavior of using the global foliage material list and instead use the Override Foliage Terrain Materials defined here.
+	 * If the override is used, you must manually provide ALL materials this procedural foliage spawner should spawn on top of.
+	 */
+	UPROPERTY(Category = ProceduralFoliageSimulation, EditAnywhere)
+	bool bUseOverrideFoliageTerrainMaterials = false;
+
+	/** Procedural foliage will only spawn on materials specified here. These are only used if 'Use Override Foliage Terrain Materials' is checked. */
+	UPROPERTY(Category = ProceduralFoliageSimulation, EditAnywhere)
+	TArray<TSoftObjectPtr<UMaterialInterface>> OverrideFoliageTerrainMaterials;
+
 public:
 	UFUNCTION(BlueprintCallable, Category = ProceduralFoliageSimulation)
 	void Simulate(int32 NumSteps = -1);
@@ -50,6 +61,10 @@ public:
 
 	const TArray<FFoliageTypeObject>& GetFoliageTypes() const { return FoliageTypes; }
 
+	bool UsesOverrideFoliageTerrainMaterials() const { return bUseOverrideFoliageTerrainMaterials; }
+
+	const TArray<TSoftObjectPtr<UMaterialInterface>>& GetFoliageTerrainMaterials() const { return OverrideFoliageTerrainMaterials; }
+
 	/** Returns the instances that need to spawn for a given min,max region */
 	void GetInstancesToSpawn(TArray<FProceduralFoliageInstance>& OutInstances, const FVector& Min = FVector(-FLT_MAX, -FLT_MAX, -FLT_MAX), const FVector& Max = FVector(FLT_MAX, FLT_MAX, FLT_MAX) ) const;
 
