Fixed mismatched copy constructor/assignment operator in FProceduralFoliageInstance.

diff --git a/Engine/Source/Runtime/Foliage/Private/ProceduralFoliageInstance.cpp b/Engine/Source/Runtime/Foliage/Private/ProceduralFoliageInstance.cpp
--- a/Engine/Source/Runtime/Foliage/Private/ProceduralFoliageInstance.cpp
+++ b/Engine/Source/Runtime/Foliage/Private/ProceduralFoliageInstance.cpp
@@ -18,20 +18,6 @@ FProceduralFoliageInstance::FProceduralFoliageInstance()
 {
 }
 
-FProceduralFoliageInstance::FProceduralFoliageInstance(const FProceduralFoliageInstance& Other)
-: Rotation(Other.Rotation)
-, Location(Other.Location)
-, Age(Other.Age)
-, Normal(Other.Normal)
-, Scale(Other.Scale)
-, Type(Other.Type)
-, BaseComponent(Other.BaseComponent)
-, bBlocker(Other.bBlocker)
-, bAlive(Other.bAlive)
-{
-
-}
-
 FProceduralFoliageInstance* GetLessFit(FProceduralFoliageInstance* A, FProceduralFoliageInstance* B)
 {
 	//Blocker is used for culling instances when we overlap tiles. It always wins
diff --git a/Engine/Source/Runtime/Foliage/Public/ProceduralFoliageInstance.h b/Engine/Source/Runtime/Foliage/Public/ProceduralFoliageInstance.h
--- a/Engine/Source/Runtime/Foliage/Public/ProceduralFoliageInstance.h
+++ b/Engine/Source/Runtime/Foliage/Public/ProceduralFoliageInstance.h
@@ -60,7 +60,6 @@ struct FOLIAGE_API FProceduralFoliageInstance
 public:
 	GENERATED_USTRUCT_BODY()
 	FProceduralFoliageInstance();
-	FProceduralFoliageInstance(const FProceduralFoliageInstance& Other);
 
 	static FProceduralFoliageInstance* Domination(FProceduralFoliageInstance* A, FProceduralFoliageInstance* B, ESimulationOverlap::Type OverlapType);
 
