Fixed mismatched copy constructor/assignment operator for FProceduralFoliageBroadphase.

diff --git a/Engine/Source/Runtime/Foliage/Private/ProceduralFoliageBroadphase.cpp b/Engine/Source/Runtime/Foliage/Private/ProceduralFoliageBroadphase.cpp
--- a/Engine/Source/Runtime/Foliage/Private/ProceduralFoliageBroadphase.cpp
+++ b/Engine/Source/Runtime/Foliage/Private/ProceduralFoliageBroadphase.cpp
@@ -14,6 +14,9 @@ FProceduralFoliageBroadphase::FProceduralFoliageBroadphase(const FProceduralFoli
 	OtherBroadphase.QuadTree.Duplicate(QuadTree);
 }
 
+// Perhaps this should copy the QuadTree like the copy constructor does...?
+FProceduralFoliageBroadphase& FProceduralFoliageBroadphase::operator=(const FProceduralFoliageBroadphase& OtherBroadphase) = default;
+
 void FProceduralFoliageBroadphase::Empty()
 {
 	QuadTree.Empty();
diff --git a/Engine/Source/Runtime/Foliage/Public/ProceduralFoliageBroadphase.h b/Engine/Source/Runtime/Foliage/Public/ProceduralFoliageBroadphase.h
--- a/Engine/Source/Runtime/Foliage/Public/ProceduralFoliageBroadphase.h
+++ b/Engine/Source/Runtime/Foliage/Public/ProceduralFoliageBroadphase.h
@@ -11,6 +11,7 @@ class FProceduralFoliageBroadphase
 public:
 	FProceduralFoliageBroadphase(float TileSize = 0.f, float MinimumQuadTreeSize = 100.f);
 	FProceduralFoliageBroadphase(const FProceduralFoliageBroadphase& OtherBroadphase);
+	FProceduralFoliageBroadphase& operator=(const FProceduralFoliageBroadphase& OtherBroadphase);
 
 	bool TestAgainstAABB(FProceduralFoliageInstance* Instance);
 	void Insert(FProceduralFoliageInstance* Instance);
