Foliage: Fix Foliage not following base component in WP map
- Move management into editor subsystem to split the moving and partitioning in two phases which fixes the issue

#rnx

diff --git a/Engine/Source/Editor/FoliageEdit/FoliageEdit.build.cs b/Engine/Source/Editor/FoliageEdit/FoliageEdit.build.cs
--- a/Engine/Source/Editor/FoliageEdit/FoliageEdit.build.cs
+++ b/Engine/Source/Editor/FoliageEdit/FoliageEdit.build.cs
@@ -14,6 +14,7 @@ public class FoliageEdit : ModuleRules
 				"InputCore",
 				"Engine",
 				"EditorFramework",
+				"EditorSubsystem",
 				"UnrealEd",
 				"Slate",
 				"SlateCore",
diff --git a/Engine/Source/Editor/FoliageEdit/Private/FoliageEditorSubsystem.cpp b/Engine/Source/Editor/FoliageEdit/Private/FoliageEditorSubsystem.cpp
new file mode 100644
--- /dev/null
+++ b/Engine/Source/Editor/FoliageEdit/Private/FoliageEditorSubsystem.cpp
@@ -0,0 +1,185 @@
+// Copyright Epic Games, Inc. All Rights Reserved.
+
+#include "FoliageEditorSubsystem.h"
+#include "CoreGlobals.h"
+#include "Engine/Engine.h"
+#include "Engine/World.h"
+#include "EngineUtils.h"
+#include "InstancedFoliageActor.h"
+#include "FoliageHelper.h"
+#include "ActorPartition/ActorPartitionSubsystem.h"
+
+#include UE_INLINE_GENERATED_CPP_BY_NAME(FoliageEditorSubsystem)
+
+#define LOCTEXT_NAMESPACE "FoliageEditorSubsystem"
+
+UFoliageEditorSubsystem::UFoliageEditorSubsystem()
+{
+}
+
+void UFoliageEditorSubsystem::Initialize(FSubsystemCollectionBase& Collection)
+{
+	Super::Initialize(Collection);
+
+	GEngine->OnActorMoved().AddUObject(this, &UFoliageEditorSubsystem::OnActorMoved);
+	GEngine->OnLevelActorDeleted().AddUObject(this, &UFoliageEditorSubsystem::OnActorDeleted);
+	GEngine->OnLevelActorOuterChanged().AddUObject(this, &UFoliageEditorSubsystem::OnActorOuterChanged);
+	FWorldDelegates::PostApplyLevelOffset.AddUObject(this, &UFoliageEditorSubsystem::OnPostApplyLevelOffset);
+	FWorldDelegates::PostApplyLevelTransform.AddUObject(this, &UFoliageEditorSubsystem::OnPostApplyLevelTransform);
+	FWorldDelegates::OnPostWorldInitialization.AddUObject(this, &UFoliageEditorSubsystem::OnPostWorldInitialization);
+
+}
+
+void UFoliageEditorSubsystem::Deinitialize()
+{
+	GEngine->OnActorMoved().RemoveAll(this);
+	GEngine->OnLevelActorDeleted().RemoveAll(this);
+	GEngine->OnLevelActorOuterChanged().RemoveAll(this);
+
+	FWorldDelegates::PostApplyLevelOffset.RemoveAll(this);
+	FWorldDelegates::PostApplyLevelTransform.RemoveAll(this);
+	FWorldDelegates::OnPostWorldInitialization.RemoveAll(this);
+
+	Super::Deinitialize();
+}
+
+void UFoliageEditorSubsystem::OnPostApplyLevelOffset(ULevel* InLevel, UWorld* InWorld, const FVector& InOffset, bool bWorldShift)
+{
+	if (InWorld && !InWorld->IsGameWorld())
+	{
+		for (TActorIterator<AInstancedFoliageActor> It(InWorld); It; ++It)
+		{
+			AInstancedFoliageActor* IFA = *It;
+			if (IFA->GetLevel() == InLevel)
+			{
+				IFA->PostApplyLevelOffset(InOffset, bWorldShift);
+			}
+		}
+	}
+}
+
+void UFoliageEditorSubsystem::OnPostApplyLevelTransform(ULevel* InLevel, const FTransform& InTransform)
+{
+	if (InLevel && InLevel->GetWorld() && !InLevel->GetWorld()->IsGameWorld())
+	{
+		for (TActorIterator<AInstancedFoliageActor> It(InLevel->GetWorld()); It; ++It)
+		{
+			AInstancedFoliageActor* IFA = *It;
+			if (IFA->GetLevel() == InLevel)
+			{
+				IFA->PostApplyLevelTransform(InTransform);
+			}
+		}
+	}
+}
+
+void UFoliageEditorSubsystem::OnPostWorldInitialization(UWorld* InWorld, const UWorld::InitializationValues IVS)
+{
+	for (TActorIterator<AInstancedFoliageActor> It(InWorld); It; ++It)
+	{
+		AInstancedFoliageActor* IFA = *It;
+		IFA->DetectFoliageTypeChangeAndUpdate();
+	}
+}
+
+void UFoliageEditorSubsystem::OnActorMoved(AActor* InActor)
+{
+	UWorld* InWorld = InActor->GetWorld();
+	if (!InWorld || InWorld->IsGameWorld())
+	{
+		return;
+	}
+
+	TMap<AInstancedFoliageActor*, TArray<UActorComponent*>> UpdatedIFAsPerComponent;
+	
+	// Process new instance positions first
+	for (TActorIterator<AInstancedFoliageActor> It(InWorld); It; ++It)
+	{
+		AInstancedFoliageActor* IFA = *It;
+		for (UActorComponent* Component : InActor->GetComponents())
+		{
+			if (IFA->MoveInstancesForMovedComponent(Component))
+			{
+				UpdatedIFAsPerComponent.FindOrAdd(IFA).Add(Component);
+			}
+		}
+
+		if (FFoliageHelper::IsOwnedByFoliage(InActor))
+		{
+			IFA->MoveInstancesForMovedOwnedActors(InActor);
+		}
+	}
+
+	// Then update partitioning
+	for (const auto& [IFA, Components] : UpdatedIFAsPerComponent)
+	{
+		for (UActorComponent* Component : Components)
+		{
+			IFA->UpdateInstancePartitioningForMovedComponent(Component);
+		}
+	}
+}
+
+void UFoliageEditorSubsystem::OnActorOuterChanged(AActor* InActor, UObject* OldOuter)
+{
+	if (GIsTransacting)
+	{
+		return;
+	}
+
+	ULevel* OldLevel = Cast<ULevel>(OldOuter);
+
+	if (InActor->GetLevel() == OldLevel)
+	{
+		return;
+	}
+
+	if (!FFoliageHelper::IsOwnedByFoliage(InActor))
+	{
+		return;
+	}
+
+	UActorPartitionSubsystem* ActorPartitionSubsystem = OldLevel ? UWorld::GetSubsystem<UActorPartitionSubsystem>(OldLevel->GetWorld()) : nullptr;
+	if (ActorPartitionSubsystem && ActorPartitionSubsystem->IsLevelPartition())
+	{
+		AInstancedFoliageActor* OldIFA = AInstancedFoliageActor::GetInstancedFoliageActorForLevel(OldLevel, false);
+		check(OldIFA);
+
+		if (OldIFA)
+		{
+			OldIFA->UpdateFoliageActorInstance(InActor);
+		}
+	}
+}
+
+void UFoliageEditorSubsystem::OnActorDeleted(AActor* InActor)
+{
+	if (GIsReinstancing)
+	{
+		return;
+	}
+	
+	UWorld* InWorld = InActor->GetWorld();
+	if (!InWorld || InWorld->IsGameWorld())
+	{
+		return;
+	}
+
+	// Process new instance positions first
+	for (TActorIterator<AInstancedFoliageActor> It(InWorld); It; ++It)
+	{
+		AInstancedFoliageActor* IFA = *It;
+		for (UActorComponent* Component : InActor->GetComponents())
+		{
+			IFA->DeleteInstancesForComponent(Component);
+		}
+
+		if (FFoliageHelper::IsOwnedByFoliage(InActor))
+		{
+			IFA->DeleteFoliageActorInstance(InActor);
+		}
+	}
+}
+
+#undef LOCTEXT_NAMESPACE
+
diff --git a/Engine/Source/Editor/FoliageEdit/Public/FoliageEditorSubsystem.h b/Engine/Source/Editor/FoliageEdit/Public/FoliageEditorSubsystem.h
new file mode 100644
--- /dev/null
+++ b/Engine/Source/Editor/FoliageEdit/Public/FoliageEditorSubsystem.h
@@ -0,0 +1,29 @@
+// Copyright Epic Games, Inc. All Rights Reserved.
+
+#pragma once
+
+#include "EditorSubsystem.h"
+#include "Engine/World.h"
+#include "FoliageEditorSubsystem.generated.h"
+
+class ULevel;
+class AActor;
+
+UCLASS()
+class UFoliageEditorSubsystem : public UEditorSubsystem
+{
+	GENERATED_BODY()
+
+	UFoliageEditorSubsystem();
+public:
+	virtual void Initialize(FSubsystemCollectionBase& Collection) override;
+	virtual void Deinitialize() override;
+
+private:
+	void OnActorMoved(AActor* InActor);
+	void OnActorOuterChanged(AActor* InActor, UObject* OldOuter);
+	void OnActorDeleted(AActor* InActor);
+	void OnPostApplyLevelOffset(ULevel* InLevel, UWorld* InWorld, const FVector& InOffset, bool bWorldShift);
+	void OnPostApplyLevelTransform(ULevel* InLevel, const FTransform& InTransform);
+	void OnPostWorldInitialization(UWorld* InWorld, const UWorld::InitializationValues IVS);
+};
diff --git a/Engine/Source/Runtime/Engine/Classes/Engine/World.h b/Engine/Source/Runtime/Engine/Classes/Engine/World.h
--- a/Engine/Source/Runtime/Engine/Classes/Engine/World.h
+++ b/Engine/Source/Runtime/Engine/Classes/Engine/World.h
@@ -4232,6 +4232,10 @@ public:
 	DECLARE_MULTICAST_DELEGATE_FourParams(FLevelOffsetEvent, ULevel*,  UWorld*, const FVector&, bool);
 	static FLevelOffsetEvent		PostApplyLevelOffset;
 
+	// Called after transform is applied to level
+	DECLARE_MULTICAST_DELEGATE_TwoParams(FLevelTransformEvent, ULevel*, const FTransform&);
+	static FLevelTransformEvent		PostApplyLevelTransform;
+
 	// called by UWorld::GetAssetRegistryTags()
 	static FWorldGetAssetTags GetAssetTags;
 
diff --git a/Engine/Source/Runtime/Engine/Private/LevelUtils.cpp b/Engine/Source/Runtime/Engine/Private/LevelUtils.cpp
--- a/Engine/Source/Runtime/Engine/Private/LevelUtils.cpp
+++ b/Engine/Source/Runtime/Engine/Private/LevelUtils.cpp
@@ -525,6 +525,7 @@ void FLevelUtils::ApplyLevelTransform(const FLevelUtils::FApplyLevelTransformPar
 #endif // WITH_EDITOR
 
 		TransformParams.Level->OnApplyLevelTransform.Broadcast(TransformParams.LevelTransform);
+		FWorldDelegates::PostApplyLevelTransform.Broadcast(TransformParams.Level, TransformParams.LevelTransform);
 	}
 }
 
diff --git a/Engine/Source/Runtime/Engine/Private/World.cpp b/Engine/Source/Runtime/Engine/Private/World.cpp
--- a/Engine/Source/Runtime/Engine/Private/World.cpp
+++ b/Engine/Source/Runtime/Engine/Private/World.cpp
@@ -553,6 +553,7 @@ FWorldDelegates::FOnLevelChanged FWorldDelegates::LevelAddedToWorld;
 FWorldDelegates::FOnLevelChanged FWorldDelegates::PreLevelRemovedFromWorld;
 FWorldDelegates::FOnLevelChanged FWorldDelegates::LevelRemovedFromWorld;
 FWorldDelegates::FLevelOffsetEvent FWorldDelegates::PostApplyLevelOffset;
+FWorldDelegates::FLevelTransformEvent FWorldDelegates::PostApplyLevelTransform;
 FWorldDelegates::FWorldGetAssetTags FWorldDelegates::GetAssetTags;
 FWorldDelegates::FOnWorldTickStart FWorldDelegates::OnWorldTickStart;
 FWorldDelegates::FOnWorldTickEnd FWorldDelegates::OnWorldTickEnd;
diff --git a/Engine/Source/Runtime/Foliage/Private/InstancedFoliage.cpp b/Engine/Source/Runtime/Foliage/Private/InstancedFoliage.cpp
--- a/Engine/Source/Runtime/Foliage/Private/InstancedFoliage.cpp
+++ b/Engine/Source/Runtime/Foliage/Private/InstancedFoliage.cpp
@@ -3160,7 +3160,7 @@ AInstancedFoliageActor* AInstancedFoliageActor::GetDefault(UWorld* InWorld)
 void AInstancedFoliageActor::MoveInstancesForMovedOwnedActors(AActor* InActor)
 {
 	// We don't want to handle this case when applying level transform
-	// since it's already handled in AInstancedFoliageActor::OnApplyLevelTransform
+	// since it's already handled in AInstancedFoliageActor::PostApplyLevelTransform
 	if (FLevelUtils::IsApplyingLevelTransform())
 	{
 		return;
@@ -3252,19 +3252,34 @@ void AInstancedFoliageActor::UpdateInstancePartitioning(UWorld* InWorld)
 	}
 }
 
-void AInstancedFoliageActor::MoveInstancesForMovedComponent(UActorComponent* InComponent)
+void AInstancedFoliageActor::UpdateInstancePartitioningForMovedComponent(UActorComponent* InComponent)
+{
+	const auto BaseId = InstanceBaseCache.GetInstanceBaseId(InComponent);
+	if (BaseId == FFoliageInstanceBaseCache::InvalidBaseId)
+	{
+		return;
+	}
+
+	for (auto& Pair : FoliageInfos)
+	{
+		FFoliageInfo& Info = *Pair.Value;
+		FoliagePartitioningUtils::Update(this, Info, [&]() { return Info.ComponentHash.Find(BaseId); });
+	}
+}
+
+bool AInstancedFoliageActor::MoveInstancesForMovedComponent(UActorComponent* InComponent)
 {
 	// We don't want to handle this case when applying level transform
-	// since it's already handled in AInstancedFoliageActor::OnApplyLevelTransform
+	// since it's already handled in AInstancedFoliageActor::PostApplyLevelTransform
 	if (FLevelUtils::IsApplyingLevelTransform())
 	{
-		return;
+		return false;
 	}
 		
 	const auto BaseId = InstanceBaseCache.GetInstanceBaseId(InComponent);
 	if (BaseId == FFoliageInstanceBaseCache::InvalidBaseId)
 	{
-		return;
+		return false;
 	}
 
 	const auto CurrentBaseInfo = InstanceBaseCache.GetInstanceBaseInfo(BaseId);
@@ -3272,13 +3287,13 @@ void AInstancedFoliageActor::MoveInstancesForMovedComponent(UActorComponent* InC
 	// Found an invalid base so don't try to move instances
 	if (!CurrentBaseInfo.BasePtr.IsValid())
 	{
-		return;
+		return false;
 	}
 
 	bool bFirst = true;
 	const auto NewBaseInfo = InstanceBaseCache.UpdateInstanceBaseInfoTransform(InComponent);
 
-	FMatrix DeltaTransfrom =
+	FMatrix DeltaTransform =
 		FTranslationMatrix(-CurrentBaseInfo.CachedLocation) *
 		FInverseRotationMatrix(CurrentBaseInfo.CachedRotation) *
 		FScaleMatrix(NewBaseInfo.CachedDrawScale / CurrentBaseInfo.CachedDrawScale) *
@@ -3309,7 +3324,7 @@ void AInstancedFoliageActor::MoveInstancesForMovedComponent(UActorComponent* InC
 				FMatrix NewTransform =
 					FRotationMatrix(Instance.Rotation) *
 					FTranslationMatrix(Instance.Location) *
-					DeltaTransfrom;
+					DeltaTransform;
 
 				// Extract rotation and position
 				Instance.Location = NewTransform.GetOrigin();
@@ -3324,10 +3339,10 @@ void AInstancedFoliageActor::MoveInstancesForMovedComponent(UActorComponent* InC
 
 			Info.Implementation->EndUpdate();
 			Info.Refresh(true, false);
-
-			FoliagePartitioningUtils::Update(this, Info, [&]() { return Info.ComponentHash.Find(BaseId); });
 		}
 	}
+
+	return true;
 }
 
 void AInstancedFoliageActor::DeleteInstancesForComponent(UActorComponent* InComponent)
@@ -4416,49 +4431,6 @@ void AInstancedFoliageActor::ExitEditMode()
 	}
 }
 
-void AInstancedFoliageActor::RegisterDelegates()
-{
-	ensureMsgf(IsInGameThread(), TEXT("Potential race condition in AInstancedFoliageActor registering to delegates from non game-thread"));
-
-	GEngine->OnActorMoved().Remove(OnLevelActorMovedDelegateHandle);
-	OnLevelActorMovedDelegateHandle = GEngine->OnActorMoved().AddUObject(this, &AInstancedFoliageActor::OnLevelActorMoved);
-
-	GEngine->OnLevelActorDeleted().Remove(OnLevelActorDeletedDelegateHandle);
-	OnLevelActorDeletedDelegateHandle = GEngine->OnLevelActorDeleted().AddUObject(this, &AInstancedFoliageActor::OnLevelActorDeleted);
-
-	if (GetLevel())
-	{
-		OnApplyLevelTransformDelegateHandle = GetLevel()->OnApplyLevelTransform.AddUObject(this, &AInstancedFoliageActor::OnApplyLevelTransform);
-	}
-
-	GEngine->OnLevelActorOuterChanged().Remove(OnLevelActorOuterChangedDelegateHandle);
-	OnLevelActorOuterChangedDelegateHandle = GEngine->OnLevelActorOuterChanged().AddUObject(this, &AInstancedFoliageActor::OnLevelActorOuterChanged);
-
-	FWorldDelegates::PostApplyLevelOffset.Remove(OnPostApplyLevelOffsetDelegateHandle);
-	OnPostApplyLevelOffsetDelegateHandle = FWorldDelegates::PostApplyLevelOffset.AddUObject(this, &AInstancedFoliageActor::OnPostApplyLevelOffset);
-
-	FWorldDelegates::OnPostWorldInitialization.Remove(OnPostWorldInitializationDelegateHandle);
-	OnPostWorldInitializationDelegateHandle = FWorldDelegates::OnPostWorldInitialization.AddUObject(this, &AInstancedFoliageActor::OnPostWorldInitialization);
-}
-
-void AInstancedFoliageActor::UnregisterDelegates()
-{
-	ensureMsgf(IsInGameThread(), TEXT("Potential race condition in AInstancedFoliageActor registering to delegates from non game-thread"));
-
-	GEngine->OnActorMoved().Remove(OnLevelActorMovedDelegateHandle);
-	GEngine->OnLevelActorDeleted().Remove(OnLevelActorDeletedDelegateHandle);
-	GEngine->OnLevelActorOuterChanged().Remove(OnLevelActorOuterChangedDelegateHandle);
-
-	if (GetLevel())
-	{
-		GetLevel()->OnApplyLevelTransform.Remove(OnApplyLevelTransformDelegateHandle);
-	}
-
-	FWorldDelegates::PostApplyLevelOffset.Remove(OnPostApplyLevelOffsetDelegateHandle);
-
-	FWorldDelegates::OnPostWorldInitialization.Remove(OnPostWorldInitializationDelegateHandle);
-}
-
 void AInstancedFoliageActor::PostInitProperties()
 {
 	Super::PostInitProperties();
@@ -4469,10 +4441,6 @@ void AInstancedFoliageActor::PostInitProperties()
 		{
 			// Delegate registration is not thread-safe, so we postpone it on PostLoad when coming from loading which could be on another thread
 		}
-		else
-		{
-			RegisterDelegates();
-		}
 	}
 }
 
@@ -4482,8 +4450,6 @@ void AInstancedFoliageActor::BeginDestroy()
 
 	if (!IsTemplate())
 	{
-		UnregisterDelegates();
-
 		FoliageInfos.Empty();
 	}
 }
@@ -4500,13 +4466,6 @@ void AInstancedFoliageActor::PostLoad()
 {
 	Super::PostLoad();
 
-#if WITH_EDITOR
-	if (!IsTemplate())
-	{
-		RegisterDelegates();
-	}
-#endif
-
 	ULevel* OwningLevel = GetLevel();
 	// We can't check the ActorPartitionSubsystem here because World is not initialized yet. So we fallback on the bIsPartitioned
 	// to know if multiple InstanceFoliageActors is valid or not.
@@ -5005,131 +4964,65 @@ void AInstancedFoliageActor::NotifyFoliageTypeWillChange(UFoliageType* FoliageTy
 	}
 }
 
-void AInstancedFoliageActor::OnLevelActorMoved(AActor* InActor)
+void AInstancedFoliageActor::UpdateFoliageActorInstance(AActor* InActor)
 {
-	UWorld* InWorld = InActor->GetWorld();
+	check(FFoliageHelper::IsOwnedByFoliage(InActor));
 
-	if (!InWorld || !InWorld->IsGameWorld())
+	TSet<int32> InstanceToMove;
+	FFoliageInfo* OldFoliageInfo = nullptr;
+	UFoliageType* FoliageType = nullptr;
+
+	for (auto& Pair : FoliageInfos)
 	{
-		for (UActorComponent* Component : InActor->GetComponents())
+		if (Pair.Value->Type == EFoliageImplType::Actor)
 		{
-			if (Component)
+			FFoliageActor* FoliageActor = StaticCast<FFoliageActor*>(Pair.Value->Implementation.Get());
+			int32 Index = FoliageActor->FindIndex(InActor);
+			if (Index != INDEX_NONE)
 			{
-				MoveInstancesForMovedComponent(Component);
+				InstanceToMove.Add(Index);
+				OldFoliageInfo = &Pair.Value.Get();
+				FoliageType = Pair.Key;
+				break;
 			}
 		}
-
-		if (FFoliageHelper::IsOwnedByFoliage(InActor))
-		{
-			MoveInstancesForMovedOwnedActors(InActor);
-		}
-	}
-}
-
-void AInstancedFoliageActor::OnLevelActorOuterChanged(AActor* InActor, UObject* OldOuter)
-{
-	if (GIsTransacting)
-	{
-		return;
-	}
-
-	ULevel* OldLevel = Cast<ULevel>(OldOuter);
-
-	if (InActor->GetLevel() == OldLevel)
-	{
-		return;
-	}
-
-	if (!FFoliageHelper::IsOwnedByFoliage(InActor))
-	{
-		return;
 	}
 
-	UActorPartitionSubsystem* ActorPartitionSubsystem = OldLevel ? UWorld::GetSubsystem<UActorPartitionSubsystem>(OldLevel->GetWorld()) : nullptr;
-	if (ActorPartitionSubsystem && ActorPartitionSubsystem->IsLevelPartition())
+	if (InstanceToMove.Num() > 0)
 	{
-		AInstancedFoliageActor* OldIFA = AInstancedFoliageActor::GetInstancedFoliageActorForLevel(OldLevel, false);
-		check(OldIFA);
-
-		if (OldIFA)
-		{
-			TSet<int32> InstanceToMove;
-			FFoliageInfo* OldFoliageInfo = nullptr;
-			UFoliageType* FoliageType = nullptr;
-
-			for (auto& Pair : OldIFA->FoliageInfos)
-			{
-				if (Pair.Value->Type == EFoliageImplType::Actor)
-				{
-					FFoliageActor* FoliageActor = StaticCast<FFoliageActor*>(Pair.Value->Implementation.Get());
-					int32 Index = FoliageActor->FindIndex(InActor);
-					if (Index != INDEX_NONE)
-					{
-						InstanceToMove.Add(Index);
-						OldFoliageInfo = &Pair.Value.Get();
-						FoliageType = Pair.Key;
-						break;
-					}
-				}
-			}
-
-			if (InstanceToMove.Num() > 0)
-			{
-				OldIFA->MoveInstancesToLevel(InActor->GetLevel(), InstanceToMove, OldFoliageInfo, FoliageType, true);
-			}
-		}
+		MoveInstancesToLevel(InActor->GetLevel(), InstanceToMove, OldFoliageInfo, FoliageType, true);
 	}
 }
 
-void AInstancedFoliageActor::OnLevelActorDeleted(AActor* InActor)
+void AInstancedFoliageActor::DeleteFoliageActorInstance(AActor* InActor)
 {
-	UWorld* InWorld = InActor->GetWorld();
-
-	if (GIsReinstancing)
-	{
-		return;
-	}
+	check(FFoliageHelper::IsOwnedByFoliage(InActor));
 
-	if (!InWorld || !InWorld->IsGameWorld())
+	for (auto& Pair : FoliageInfos)
 	{
-		for (UActorComponent* Component : InActor->GetComponents())
-		{
-			if (Component)
-			{
-				DeleteInstancesForComponent(Component);
-			}
-		}
+		FFoliageInfo& Info = *Pair.Value;
+		UFoliageType* FoliageType = Pair.Key;
 
-		// Cleanup Foliage Instances if Actor is deleted outside of Foliage Tool
-		if (FFoliageHelper::IsOwnedByFoliage(InActor))
+		if (Info.Type == EFoliageImplType::Actor)
 		{
-			for (auto& Pair : FoliageInfos)
+			if (FFoliageActor* FoliageActor = StaticCast<FFoliageActor*>(Info.Implementation.Get()))
 			{
-				FFoliageInfo& Info = *Pair.Value;
-				UFoliageType* FoliageType = Pair.Key;
+				TArray<int32> InstancesToRemove;
+				// Make sure we find Null pointers and Delete those instances if we can't find the actor
+				FoliageActor->GetInvalidInstances(InstancesToRemove);
+				// Find Actor
+				int32 Index = FoliageActor->FindIndex(InActor);
+				if (Index != INDEX_NONE)
+				{
+					InstancesToRemove.Add(Index);
+				}
 
-				if (Info.Type == EFoliageImplType::Actor)
+				if (InstancesToRemove.Num() > 0)
 				{
-					if (FFoliageActor* FoliageActor = StaticCast<FFoliageActor*>(Info.Implementation.Get()))
+					Info.RemoveInstances(InstancesToRemove, false);
+					if (InstanceCountChanged.IsBound())
 					{
-						TArray<int32> InstancesToRemove;
-						// Make sure we find Null pointers and Delete those instances if we can't find the actor
-						FoliageActor->GetInvalidInstances(InstancesToRemove);
-						// Find Actor
-						int32 Index = FoliageActor->FindIndex(InActor);
-						if (Index != INDEX_NONE)
-						{
-							InstancesToRemove.Add(Index);
-						}
-						
-						if(InstancesToRemove.Num() > 0)
-						{
-							Info.RemoveInstances(InstancesToRemove, false);
-							if (InstanceCountChanged.IsBound())
-							{
-								InstanceCountChanged.Broadcast(FoliageType);
-							}
-						}
+						InstanceCountChanged.Broadcast(FoliageType);
 					}
 				}
 			}
@@ -5137,14 +5030,6 @@ void AInstancedFoliageActor::OnLevelActorDeleted(AActor* InActor)
 	}
 }
 
-void AInstancedFoliageActor::OnPostWorldInitialization(UWorld* World, const UWorld::InitializationValues IVS)
-{
-	if (GetWorld() == World)
-	{
-		DetectFoliageTypeChangeAndUpdate();
-	}
-}
-
 // This logic was extracted from AInstancedFoliageActor::PostLoad to be called once the World is done initializing
 void AInstancedFoliageActor::DetectFoliageTypeChangeAndUpdate()
 {
@@ -5205,55 +5090,45 @@ bool AInstancedFoliageActor::ShouldIncludeGridSizeInName(UWorld* InWorld, const
 	return InWorld->GetWorldSettings()->bIncludeGridSizeInNameForFoliageActors;
 }
 
-void AInstancedFoliageActor::OnApplyLevelTransform(const FTransform& InTransform)
+void AInstancedFoliageActor::PostApplyLevelTransform(const FTransform& InTransform)
 {
 #if WITH_EDITORONLY_DATA
-	if (GIsEditor)
+	check(GIsEditor);
+	InstanceBaseCache.UpdateInstanceBaseCachedTransforms();
+	for (auto& Pair : FoliageInfos)
 	{
-		InstanceBaseCache.UpdateInstanceBaseCachedTransforms();
-		for (auto& Pair : FoliageInfos)
+		FFoliageInfo& Info = *Pair.Value;
+		Info.InstanceHash->Empty();
+		for (int32 InstanceIdx = 0; InstanceIdx < Info.Instances.Num(); InstanceIdx++)
 		{
-			FFoliageInfo& Info = *Pair.Value;
-			Info.InstanceHash->Empty();
-			for (int32 InstanceIdx = 0; InstanceIdx < Info.Instances.Num(); InstanceIdx++)
-			{
-				FFoliageInstance& Instance = Info.Instances[InstanceIdx];
-				FTransform OldTransform(Instance.Rotation, Instance.Location);
-				FTransform NewTransform = OldTransform * InTransform;
-				Instance.Location = NewTransform.GetTranslation();
-				Instance.Rotation = NewTransform.Rotator();
-				// Rehash instance location
-				Info.InstanceHash->InsertInstance(Instance.Location, InstanceIdx);
-			}
+			FFoliageInstance& Instance = Info.Instances[InstanceIdx];
+			FTransform OldTransform(Instance.Rotation, Instance.Location);
+			FTransform NewTransform = OldTransform * InTransform;
+			Instance.Location = NewTransform.GetTranslation();
+			Instance.Rotation = NewTransform.Rotator();
+			// Rehash instance location
+			Info.InstanceHash->InsertInstance(Instance.Location, InstanceIdx);
 		}
 	}
 #endif
 }
 
-void AInstancedFoliageActor::OnPostApplyLevelOffset(ULevel* InLevel, UWorld* InWorld, const FVector& InOffset, bool bWorldShift)
+void AInstancedFoliageActor::PostApplyLevelOffset(const FVector& InOffset, bool bWorldShift)
 {
-	ULevel* OwningLevel = GetLevel();
-	if (InLevel != OwningLevel) // TODO: cross-level foliage 
-	{
-		return;
-	}
-
-	if (GIsEditor && InWorld && !InWorld->IsGameWorld())
+	check(GIsEditor);
+	for (auto& Pair : FoliageInfos)
 	{
-		for (auto& Pair : FoliageInfos)
-		{
-			FFoliageInfo& Info = *Pair.Value;
+		FFoliageInfo& Info = *Pair.Value;
 
-			InstanceBaseCache.UpdateInstanceBaseCachedTransforms();
+		InstanceBaseCache.UpdateInstanceBaseCachedTransforms();
 
-			Info.InstanceHash->Empty();
-			for (int32 InstanceIdx = 0; InstanceIdx < Info.Instances.Num(); InstanceIdx++)
-			{
-				FFoliageInstance& Instance = Info.Instances[InstanceIdx];
-				Instance.Location += InOffset;
-				// Rehash instance location
-				Info.InstanceHash->InsertInstance(Instance.Location, InstanceIdx);
-			}
+		Info.InstanceHash->Empty();
+		for (int32 InstanceIdx = 0; InstanceIdx < Info.Instances.Num(); InstanceIdx++)
+		{
+			FFoliageInstance& Instance = Info.Instances[InstanceIdx];
+			Instance.Location += InOffset;
+			// Rehash instance location
+			Info.InstanceHash->InsertInstance(Instance.Location, InstanceIdx);
 		}
 	}
 }
diff --git a/Engine/Source/Runtime/Foliage/Public/InstancedFoliageActor.h b/Engine/Source/Runtime/Foliage/Public/InstancedFoliageActor.h
--- a/Engine/Source/Runtime/Foliage/Public/InstancedFoliageActor.h
+++ b/Engine/Source/Runtime/Foliage/Public/InstancedFoliageActor.h
@@ -85,9 +85,6 @@ public:
 
 protected:
 #if WITH_EDITOR
-	void RegisterDelegates();
-	void UnregisterDelegates();
-
 	void HandleFoliageInstancePreMove(const FFoliageInstanceId& InstanceId);
 	void HandleFoliageInstancePostMove(const FFoliageInstanceId& InstanceId);
 #endif
@@ -212,12 +209,6 @@ public:
 	// Move instances to a foliage actor in target level
 	FOLIAGE_API void MoveInstancesToLevel(ULevel* InTargetLevel, TSet<int32>& InInstanceList, FFoliageInfo* InCurrentMeshInfo, UFoliageType* InFoliageType, bool bSelect = false);
 
-	// Move instances based on a component that has just been moved.
-	void MoveInstancesForMovedComponent(UActorComponent* InComponent);
-
-	// Move instances that are owned by foliage actor.
-	void MoveInstancesForMovedOwnedActors(AActor* InActor);
-	
 	// Returns a map of Static Meshes and their placed instances attached to a component.
 	FOLIAGE_API TMap<UFoliageType*, TArray<const FFoliageInstancePlacementInfo*>> GetInstancesForComponent(UActorComponent* InComponent);
 
@@ -306,24 +297,23 @@ private:
 	
 #if WITH_EDITOR
 	void ClearSelection();
-	void OnLevelActorMoved(AActor* InActor);
-	void OnLevelActorOuterChanged(AActor* InActor, UObject* OldOuter);
-	void OnLevelActorDeleted(AActor* InActor);
-	void OnApplyLevelTransform(const FTransform& InTransform);
-	void OnPostApplyLevelOffset(ULevel* InLevel, UWorld* InWorld, const FVector& InOffset, bool bWorldShift);
-	void OnPostWorldInitialization(UWorld* World, const UWorld::InitializationValues IVS);
 	void MoveInstancesToNewComponent(UPrimitiveComponent* InOldComponent, UPrimitiveComponent* InNewComponent, TFunctionRef<TArray<int32>(const FFoliageInfo&)> GetInstancesToMoveFunc);
 	bool DeleteInstancesForProceduralFoliageComponentInternal(const FGuid& InProceduralGuid, bool bInRebuildTree, bool bInDeleteAll);
+	
+	friend class UFoliageEditorSubsystem;
+	FOLIAGE_API bool MoveInstancesForMovedComponent(UActorComponent* InComponent);
+	FOLIAGE_API void UpdateInstancePartitioningForMovedComponent(UActorComponent* InComponent);
+	
+	FOLIAGE_API void MoveInstancesForMovedOwnedActors(AActor* InActor);
+
+	FOLIAGE_API void UpdateFoliageActorInstance(AActor* InActor);
+	FOLIAGE_API void DeleteFoliageActorInstance(AActor* InActor);
+
+	FOLIAGE_API void PostApplyLevelOffset(const FVector& InOffset, bool bWorldShift);
+	FOLIAGE_API void PostApplyLevelTransform(const FTransform& InTransform);
 #endif
 private:
 #if WITH_EDITOR
-	FDelegateHandle OnLevelActorMovedDelegateHandle;
-	FDelegateHandle OnLevelActorDeletedDelegateHandle;
-	FDelegateHandle OnLevelActorOuterChangedDelegateHandle;
-	FDelegateHandle OnPostApplyLevelOffsetDelegateHandle;
-	FDelegateHandle OnApplyLevelTransformDelegateHandle;
-	FDelegateHandle OnPostWorldInitializationDelegateHandle;
-	
 	FOnFoliageTypeMeshChanged OnFoliageTypeMeshChangedEvent;
 #endif
 
