Add checking for invalid indices and instances.

diff --git a/Engine/Source/Runtime/Foliage/Private/FoliageActor.cpp b/Engine/Source/Runtime/Foliage/Private/FoliageActor.cpp
--- a/Engine/Source/Runtime/Foliage/Private/FoliageActor.cpp
+++ b/Engine/Source/Runtime/Foliage/Private/FoliageActor.cpp
@@ -147,17 +147,31 @@ void FFoliageActor::AddExistingInstance(const FFoliageInstance& ExistingInstance
 
 void FFoliageActor::RemoveInstance(int32 InstanceIndex)
 {
-	AActor* Actor = ActorInstances[InstanceIndex];
-	ActorInstances.RemoveAtSwap(InstanceIndex);
-	Actor->GetWorld()->DestroyActor(Actor, true);
-	bActorsDestroyed = true;
+	if (ActorInstances.IsValidIndex(InstanceIndex))
+	{
+		if (AActor* Actor = ActorInstances[InstanceIndex])
+		{
+			Actor->GetWorld()->DestroyActor(Actor, true);
+			bActorsDestroyed = true;
+		}
+
+		ActorInstances.RemoveAtSwap(InstanceIndex);
+	}
 }
 
 void FFoliageActor::MoveInstance(int32 InstanceIndex, UObject*& OutInstanceImplementation)
 {
-	AActor* Actor = ActorInstances[InstanceIndex];
-	OutInstanceImplementation = Actor;
-	ActorInstances.RemoveAtSwap(InstanceIndex);
+	OutInstanceImplementation = nullptr;
+
+	if (ActorInstances.IsValidIndex(InstanceIndex))
+	{
+		if (AActor* Actor = ActorInstances[InstanceIndex])
+		{
+			OutInstanceImplementation = Actor;
+		}
+
+		ActorInstances.RemoveAtSwap(InstanceIndex);
+	}
 }
 
 void FFoliageActor::BeginUpdate()
@@ -179,17 +193,23 @@ void FFoliageActor::EndUpdate()
 
 void FFoliageActor::SetInstanceWorldTransform(int32 InstanceIndex, const FTransform& Transform, bool bTeleport)
 {
-	if (AActor* Actor = ActorInstances[InstanceIndex])
+	if (ActorInstances.IsValidIndex(InstanceIndex))
 	{
-		Actor->SetActorTransform(Transform);
+		if (AActor* Actor = ActorInstances[InstanceIndex])
+		{
+			Actor->SetActorTransform(Transform);
+		}
 	}
 }
 
 FTransform FFoliageActor::GetInstanceWorldTransform(int32 InstanceIndex) const
 {
-	if (AActor* Actor = ActorInstances[InstanceIndex])
+	if (ActorInstances.IsValidIndex(InstanceIndex))
 	{
-		return Actor->GetTransform();
+		if (AActor* Actor = ActorInstances[InstanceIndex])
+		{
+			return Actor->GetTransform();
+		}
 	}
 
 	return FTransform::Identity;
@@ -345,9 +365,15 @@ void FFoliageActor::SelectAllInstances(bool bSelect)
 
 void FFoliageActor::SelectInstance(bool bSelect, int32 Index)
 {
-	TArray<AActor*> SingleInstance;
-	SingleInstance.Add(ActorInstances[Index]);
-	AInstancedFoliageActor::SelectionChanged.Broadcast(bSelect, SingleInstance);
+	if (ActorInstances.IsValidIndex(Index))
+	{
+		if (AActor* Actor = ActorInstances[Index])
+		{
+			TArray<AActor*> SingleInstance;
+			SingleInstance.Add(Actor);
+			AInstancedFoliageActor::SelectionChanged.Broadcast(bSelect, SingleInstance);
+		}
+	}
 }
 
 void FFoliageActor::SelectInstances(bool bSelect, const TSet<int32>& SelectedIndices)
@@ -380,20 +406,27 @@ void FFoliageActor::ClearSelection(const TSet<int32>& SelectedIndices)
 
 bool FFoliageActor::UpdateInstanceFromActor(int32 Index, FFoliageInfo& FoliageInfo)
 {
-	AActor* Actor = ActorInstances[Index];
-	AInstancedFoliageActor* IFA = GetIFA();
-	IFA->Modify();
-	const bool bChecked = false; // In the case of the PostEditUndo its possible that the instancehash is empty.
-	FoliageInfo.InstanceHash->RemoveInstance(FoliageInfo.Instances[Index].Location, Index, bChecked);
-	
-	FTransform ActorTransform = Actor->GetTransform();
-	FoliageInfo.Instances[Index].Location = ActorTransform.GetLocation();
-	FoliageInfo.Instances[Index].Rotation = FRotator(ActorTransform.GetRotation());
-	FoliageInfo.Instances[Index].PreAlignRotation = FoliageInfo.Instances[Index].Rotation;
-	FoliageInfo.Instances[Index].DrawScale3D = (FVector3f)Actor->GetActorScale3D();
-	FoliageInfo.InstanceHash->InsertInstance(FoliageInfo.Instances[Index].Location, Index);
-	
-	return true;
+	if (ActorInstances.IsValidIndex(Index))
+	{
+		if (AActor* Actor = ActorInstances[Index])
+		{
+			AInstancedFoliageActor* IFA = GetIFA();
+			IFA->Modify();
+			const bool bChecked = false; // In the case of the PostEditUndo its possible that the instancehash is empty.
+			FoliageInfo.InstanceHash->RemoveInstance(FoliageInfo.Instances[Index].Location, Index, bChecked);
+
+			FTransform ActorTransform = Actor->GetTransform();
+			FoliageInfo.Instances[Index].Location = ActorTransform.GetLocation();
+			FoliageInfo.Instances[Index].Rotation = FRotator(ActorTransform.GetRotation());
+			FoliageInfo.Instances[Index].PreAlignRotation = FoliageInfo.Instances[Index].Rotation;
+			FoliageInfo.Instances[Index].DrawScale3D = (FVector3f)Actor->GetActorScale3D();
+			FoliageInfo.InstanceHash->InsertInstance(FoliageInfo.Instances[Index].Location, Index);
+
+			return true;
+		}
+	}
+
+	return false;
 }
 
 void FFoliageActor::GetInvalidInstances(TArray<int32>& InvalidInstances)
