Foliage Fixes:
- Use FScopedActorEditorContextFromActor for operations which need to use the same context as the Source actor
- Update Move Tool to allow Move to current Actor Editor Context. (Allows changing DataLayers / ContentBundle for selected instances)

#rnx

diff --git a/Engine/Source/Editor/EditorStyle/Private/SlateEditorStyle.cpp b/Engine/Source/Editor/EditorStyle/Private/SlateEditorStyle.cpp
--- a/Engine/Source/Editor/EditorStyle/Private/SlateEditorStyle.cpp
+++ b/Engine/Source/Editor/EditorStyle/Private/SlateEditorStyle.cpp
@@ -1781,6 +1781,8 @@ void FSlateEditorStyle::FStyle::SetupGeneralStyles()
 		Set("FoliageEditMode.Settings.Small",            new IMAGE_BRUSH("Icons/GeneralTools/Settings_40x", Icon20x20));
 		Set("FoliageEditMode.MoveToCurrentLevel",        new IMAGE_BRUSH("Icons/GeneralTools/MoveToLevel_40x", Icon20x20));
 		Set("FoliageEditMode.MoveToCurrentLevel.Small",  new IMAGE_BRUSH("Icons/GeneralTools/MoveToLevel_40x", Icon20x20));
+		Set("FoliageEditMode.MoveToActorEditorContext", new IMAGE_BRUSH("Icons/GeneralTools/MoveToActorEditorContext_40x", Icon20x20));
+		Set("FoliageEditMode.MoveToActorEditorContext.Small", new IMAGE_BRUSH("Icons/GeneralTools/MoveToActorEditorContext_40x", Icon20x20));
 
 		Set( "FoliageEditMode.SetNoSettings", new IMAGE_BRUSH( "Icons/FoliageEditMode/icon_FoliageEdMode_NoSettings_20x", Icon20x20 ) );
 		Set( "FoliageEditMode.SetPaintSettings", new IMAGE_BRUSH( "Icons/FoliageEditMode/icon_FoliageEdMode_PaintingSettings_20x", Icon20x20 ) );
diff --git a/Engine/Source/Editor/EditorStyle/Private/StarshipStyle.cpp b/Engine/Source/Editor/EditorStyle/Private/StarshipStyle.cpp
--- a/Engine/Source/Editor/EditorStyle/Private/StarshipStyle.cpp
+++ b/Engine/Source/Editor/EditorStyle/Private/StarshipStyle.cpp
@@ -1501,6 +1501,8 @@ void FStarshipEditorStyle::FStyle::SetupGeneralStyles()
 		Set("FoliageEditMode.Settings.Small",            new IMAGE_BRUSH("Icons/GeneralTools/Settings_40x", Icon20x20));
 		Set("FoliageEditMode.MoveToCurrentLevel",        new IMAGE_BRUSH("Icons/GeneralTools/MoveToLevel_40x", Icon20x20));
 		Set("FoliageEditMode.MoveToCurrentLevel.Small",  new IMAGE_BRUSH("Icons/GeneralTools/MoveToLevel_40x", Icon20x20));
+		Set("FoliageEditMode.MoveToActorEditorContext", new IMAGE_BRUSH("Icons/GeneralTools/MoveToActorEditorContext_40x", Icon20x20));
+		Set("FoliageEditMode.MoveToActorEditorContext.Small", new IMAGE_BRUSH("Icons/GeneralTools/MoveToActorEditorContext_40x", Icon20x20));
 
 		Set( "FoliageEditMode.SetNoSettings", new IMAGE_BRUSH( "Icons/FoliageEditMode/icon_FoliageEdMode_NoSettings_20x", Icon20x20 ) );
 		Set( "FoliageEditMode.SetPaintSettings", new IMAGE_BRUSH( "Icons/FoliageEditMode/icon_FoliageEdMode_PaintingSettings_20x", Icon20x20 ) );
diff --git a/Engine/Source/Editor/FoliageEdit/Private/FoliageEdMode.cpp b/Engine/Source/Editor/FoliageEdit/Private/FoliageEdMode.cpp
--- a/Engine/Source/Editor/FoliageEdit/Private/FoliageEdMode.cpp
+++ b/Engine/Source/Editor/FoliageEdit/Private/FoliageEdMode.cpp
@@ -3325,6 +3325,11 @@ bool FEdModeFoliage::IsModifierButtonPressed(const FEditorViewportClient* Viewpo
 	return IsShiftDown(ViewportClient->Viewport);
 }
 
+void FEdModeFoliage::MoveSelectedFoliageToActorEditorContext()
+{
+	AInstancedFoliageActor::MoveSelectedInstancesToActorEditorContext(GetWorld());
+}
+
 bool FEdModeFoliage::CanMoveSelectedFoliageToLevel(ULevel* InTargetLevel) const
 {
 	UWorld* World = InTargetLevel->OwningWorld;
diff --git a/Engine/Source/Editor/FoliageEdit/Private/FoliageEdMode.h b/Engine/Source/Editor/FoliageEdit/Private/FoliageEdMode.h
--- a/Engine/Source/Editor/FoliageEdit/Private/FoliageEdMode.h
+++ b/Engine/Source/Editor/FoliageEdit/Private/FoliageEdMode.h
@@ -419,6 +419,9 @@ public:
 
 	/** Tell us if we can moves selected foliage instances to the target level. */
 	bool CanMoveSelectedFoliageToLevel(ULevel* InTargetLevel) const;
+	
+	/** Moves selected foliage instances to the current actor editor context */
+	void MoveSelectedFoliageToActorEditorContext();
 
 	/** Ends tracking and end potential transaction */
 	bool EndTracking();
diff --git a/Engine/Source/Editor/FoliageEdit/Private/SFoliageEdit.cpp b/Engine/Source/Editor/FoliageEdit/Private/SFoliageEdit.cpp
--- a/Engine/Source/Editor/FoliageEdit/Private/SFoliageEdit.cpp
+++ b/Engine/Source/Editor/FoliageEdit/Private/SFoliageEdit.cpp
@@ -587,11 +587,11 @@ void SFoliageEdit::CustomizeToolBarPalette(FToolBarBuilder& ToolBarBuilder)
 
 	// Move To Current Level
 	ToolBarBuilder.AddToolBarButton(
-		FExecuteAction::CreateSP(this, &SFoliageEdit::OnMoveSelectedInstancesToCurrentLevel),
+		FExecuteAction::CreateSP(this, &SFoliageEdit::OnMoveSelectedInstancesToActorEditorContext),
 		NAME_None,
-		LOCTEXT("FoliageMoveToCurrentLevel", "Move"),
-		LOCTEXT("FoliageMoveToCurrentLevelTooltip", "Move the Selected Foliage to the Current Level"),
-		FSlateIcon(FAppStyle::GetAppStyleSetName(), "FoliageEditMode.MoveToCurrentLevel")
+		LOCTEXT("FoliageMoveToCurrentContext", "Move"),
+		LOCTEXT("FoliageMoveToCurrentContextTooltip", "Move the Selected Foliage to the Current Context"),
+		FSlateIcon(FAppStyle::GetAppStyleSetName(), "FoliageEditMode.MoveToActorEditorContext")
 	);
 
 }
@@ -971,13 +971,20 @@ void SFoliageEdit::OnDeselectAllInstances()
 	});
 }
 
-void SFoliageEdit::OnMoveSelectedInstancesToCurrentLevel()
+void SFoliageEdit::OnMoveSelectedInstancesToActorEditorContext()
 {
 	if (UWorld* World = FoliageEditMode->GetWorld())
 	{
-		if (ULevel* CurrentLevel = World->GetCurrentLevel())
+		if (!World->IsPartitionedWorld())
 		{
-			FoliageEditMode->MoveSelectedFoliageToLevel(CurrentLevel);
+			if (ULevel* CurrentLevel = World->GetCurrentLevel())
+			{
+				FoliageEditMode->MoveSelectedFoliageToLevel(CurrentLevel);
+			}
+		}
+		else
+		{
+			FoliageEditMode->MoveSelectedFoliageToActorEditorContext();
 		}
 	}
 }
diff --git a/Engine/Source/Editor/FoliageEdit/Private/SFoliageEdit.h b/Engine/Source/Editor/FoliageEdit/Private/SFoliageEdit.h
--- a/Engine/Source/Editor/FoliageEdit/Private/SFoliageEdit.h
+++ b/Engine/Source/Editor/FoliageEdit/Private/SFoliageEdit.h
@@ -240,8 +240,8 @@ public:	// SELECTION
 	/** Handler for 'Deselect All' command  */
 	void OnDeselectAllInstances();
 
-	/** Handler for 'Move to Current Level' command*/
-	void OnMoveSelectedInstancesToCurrentLevel();
+	/** Handler for 'Move to Current Editor Context' command*/
+	void OnMoveSelectedInstancesToActorEditorContext();
 
 	/** Tooltip text for 'Instance Count" column */
 	FText GetTotalInstanceCountTooltipText() const;
diff --git a/Engine/Source/Runtime/Foliage/Private/InstancedFoliage.cpp b/Engine/Source/Runtime/Foliage/Private/InstancedFoliage.cpp
--- a/Engine/Source/Runtime/Foliage/Private/InstancedFoliage.cpp
+++ b/Engine/Source/Runtime/Foliage/Private/InstancedFoliage.cpp
@@ -58,6 +58,10 @@ InstancedFoliage.cpp: Instanced foliage implementation.
 #include "Misc/CoreMisc.h"
 #include "Engine/DamageEvents.h"
 
+#if WITH_EDITOR
+#include "ActorEditorContext/ScopedActorEditorContextFromActor.h"
+#endif
+
 #define LOCTEXT_NAMESPACE "InstancedFoliage"
 
 #define DO_FOLIAGE_CHECK			0			// whether to validate foliage data during editing.
@@ -3246,11 +3250,25 @@ namespace FoliagePartitioningUtils
 	}
 }
 
+void AInstancedFoliageActor::MoveSelectedInstancesToActorEditorContext(UWorld* InWorld)
+{
+	for (TActorIterator<AInstancedFoliageActor> It(InWorld); It; ++It)
+	{
+		AInstancedFoliageActor* IFA = *It;
+		IFA->ForEachFoliageInfo([IFA](UFoliageType* FoliageType, FFoliageInfo& FoliageInfo)
+		{
+			FoliagePartitioningUtils::Update(IFA, FoliageInfo, [&]() { return &FoliageInfo.SelectedIndices; });
+			return true; // continue iteration
+		});
+	}
+}
+
 void AInstancedFoliageActor::UpdateInstancePartitioning(UWorld* InWorld)
 {
 	for (TActorIterator<AInstancedFoliageActor> It(InWorld); It; ++It)
 	{
 		AInstancedFoliageActor* IFA = *It;
+		FScopedActorEditorContextFromActor Context(IFA);
 		IFA->ForEachFoliageInfo([IFA](UFoliageType* FoliageType, FFoliageInfo& FoliageInfo)
         {
 			FoliagePartitioningUtils::Update(IFA, FoliageInfo, [&]() { return &FoliageInfo.SelectedIndices; });
@@ -3267,10 +3285,13 @@ void AInstancedFoliageActor::UpdateInstancePartitioningForMovedComponent(UActorC
 		return;
 	}
 
-	for (auto& Pair : FoliageInfos)
 	{
-		FFoliageInfo& Info = *Pair.Value;
-		FoliagePartitioningUtils::Update(this, Info, [&]() { return Info.ComponentHash.Find(BaseId); });
+		FScopedActorEditorContextFromActor Context(this);
+		for (auto& Pair : FoliageInfos)
+		{
+			FFoliageInfo& Info = *Pair.Value;
+			FoliagePartitioningUtils::Update(this, Info, [&]() { return Info.ComponentHash.Find(BaseId); });
+		}
 	}
 }
 
@@ -3565,6 +3586,7 @@ void AInstancedFoliageActor::MoveInstancesToNewComponent(UPrimitiveComponent* In
 
 	// If Modify was called on this IFA
 	bool bModified = false;
+	FScopedActorEditorContextFromActor Context(this);
 	
 	for (auto& Pair : FoliageInfos)
 	{			
@@ -5659,6 +5681,7 @@ void AInstancedFoliageActor::HandleFoliageInstancePostMove(const FFoliageInstanc
 	// This instance may be been moved into a new world partition
 	// Verify, and re-instance any existing typed elements if required
 	{
+		FScopedActorEditorContextFromActor Context(this);
 		const FFoliageInstance& FoliageInstance = InstanceId.GetInstanceChecked();
 		AInstancedFoliageActor* TargetIFA = AInstancedFoliageActor::Get(GetWorld(), /*bCreateIfNone*/true, GetLevel(), FoliageInstance.Location);
 		if (TargetIFA != this)
diff --git a/Engine/Source/Runtime/Foliage/Public/InstancedFoliageActor.h b/Engine/Source/Runtime/Foliage/Public/InstancedFoliageActor.h
--- a/Engine/Source/Runtime/Foliage/Public/InstancedFoliageActor.h
+++ b/Engine/Source/Runtime/Foliage/Public/InstancedFoliageActor.h
@@ -176,6 +176,7 @@ public:
 	static FOLIAGE_API AInstancedFoliageActor* GetDefault(UWorld* InWorld);
 
 	static FOLIAGE_API void UpdateInstancePartitioning(UWorld* InWorld);
+	static FOLIAGE_API void MoveSelectedInstancesToActorEditorContext(UWorld* InWorld);
 
 	static FOLIAGE_API bool FoliageTrace(const UWorld* InWorld, FHitResult& OutHit, const FDesiredFoliageInstance& DesiredInstance, FName InTraceTag = NAME_None, bool InbReturnFaceIndex = false, const FFoliageTraceFilterFunc& FilterFunc = FFoliageTraceFilterFunc(), bool bAverageNormal = false);
 	static FOLIAGE_API bool CheckCollisionWithWorld(const UWorld* InWorld, const UFoliageType* Settings, const FFoliageInstance& Inst, const FVector& HitNormal, const FVector& HitLocation, UPrimitiveComponent* HitComponent);
