Fix Foliage replace crash: from Actor (with 0 instances) to StaticMesh

#rnx

diff --git a/Engine/Source/Runtime/Foliage/Private/InstancedFoliage.cpp b/Engine/Source/Runtime/Foliage/Private/InstancedFoliage.cpp
--- a/Engine/Source/Runtime/Foliage/Private/InstancedFoliage.cpp
+++ b/Engine/Source/Runtime/Foliage/Private/InstancedFoliage.cpp
@@ -2631,7 +2631,7 @@ void FFoliageInfo::RecomputeHash()
 void FFoliageInfo::ReallocateClusters(UFoliageType* InSettings)
 {
 	// In case Foliage Type Changed recreate implementation
-	if (Implementation.IsValid())
+	if (Implementation.IsValid() && Implementation->IsInitialized())
 	{
 		Implementation->Uninitialize();
 	}
