Discard hits on hidden actors when placing foliage.

diff --git a/Engine/Source/Runtime/Foliage/Private/InstancedFoliage.cpp b/Engine/Source/Runtime/Foliage/Private/InstancedFoliage.cpp
--- a/Engine/Source/Runtime/Foliage/Private/InstancedFoliage.cpp
+++ b/Engine/Source/Runtime/Foliage/Private/InstancedFoliage.cpp
@@ -5282,22 +5282,34 @@ bool AInstancedFoliageActor::FoliageTrace(const UWorld* InWorld, FHitResult& Out
 		// Don't place foliage on itself
 		const AActor* HitActor = Hit.HitObjectHandle.FetchActor();
 		const AInstancedFoliageActor* IFA = Cast<AInstancedFoliageActor>(HitActor);
-		if (!IFA && HitActor && FFoliageHelper::IsOwnedByFoliage(HitActor))
+
+		if (HitActor)
 		{
-			IFA = HitActor->GetLevel()->InstancedFoliageActor.Get();
-			if (IFA == nullptr)
+			// Don't place foliage on hidden actors.
+			if (HitActor->IsTemporarilyHiddenInEditor())
 			{
 				bOutDiscardHit = true;
 				return true;
 			}
 
-			if (const FFoliageInfo* FoundMeshInfo = IFA->FindInfo(DesiredInstance.FoliageType))
+			// Don't place foliage on itself.
+			if (!IFA && FFoliageHelper::IsOwnedByFoliage(HitActor))
 			{
-				if (FoundMeshInfo->Implementation->IsOwnedComponent(HitComponent))
+				IFA = HitActor->GetLevel()->InstancedFoliageActor.Get();
+				if (IFA == nullptr)
 				{
 					bOutDiscardHit = true;
 					return true;
 				}
+
+				if (const FFoliageInfo* FoundMeshInfo = IFA->FindInfo(DesiredInstance.FoliageType))
+				{
+					if (FoundMeshInfo->Implementation->IsOwnedComponent(HitComponent))
+					{
+						bOutDiscardHit = true;
+						return true;
+					}
+				}
 			}
 		}
 
