Fix regression from 26087417 when deleting foliage actors.

diff --git a/Engine/Source/Runtime/Foliage/Private/FoliageActor.cpp b/Engine/Source/Runtime/Foliage/Private/FoliageActor.cpp
--- a/Engine/Source/Runtime/Foliage/Private/FoliageActor.cpp
+++ b/Engine/Source/Runtime/Foliage/Private/FoliageActor.cpp
@@ -149,13 +149,14 @@ void FFoliageActor::RemoveInstance(int32 InstanceIndex)
 {
 	if (ActorInstances.IsValidIndex(InstanceIndex))
 	{
-		if (AActor* Actor = ActorInstances[InstanceIndex])
+		AActor* Actor = ActorInstances[InstanceIndex];
+		ActorInstances.RemoveAtSwap(InstanceIndex);
+
+		if (Actor)
 		{
 			Actor->GetWorld()->DestroyActor(Actor, true);
 			bActorsDestroyed = true;
 		}
-
-		ActorInstances.RemoveAtSwap(InstanceIndex);
 	}
 }
 
