Fix UGrassInstancedStaticMeshComponent to return correct value for GetNumRenderInstances().
Previous implementation relied on SortedInstances which isn't kept by UGrassInstancedStaticMeshComponent which always receives a prebuild cluster tree, with no SortedInstances array.

diff --git a/Engine/Source/Runtime/Foliage/Public/GrassInstancedStaticMeshComponent.h b/Engine/Source/Runtime/Foliage/Public/GrassInstancedStaticMeshComponent.h
--- a/Engine/Source/Runtime/Foliage/Public/GrassInstancedStaticMeshComponent.h
+++ b/Engine/Source/Runtime/Foliage/Public/GrassInstancedStaticMeshComponent.h
@@ -21,6 +21,8 @@ public:
 
 	FOLIAGE_API bool SupportsWorldPositionOffsetVelocity() const override;
 
+	int32 GetNumRenderInstances() const override { return NumBuiltRenderInstances; }
+
 private:
 	virtual void BuildComponentInstanceData(ERHIFeatureLevel::Type FeatureLevel, FInstanceUpdateComponentDesc& OutData) override final;
 };
