Temporarily enable RF_Transactional in PreEditChange() events to allow component changes to be captured by the transaction system.

#rnx

diff --git a/Engine/Source/Runtime/Foliage/Private/InstancedFoliage.cpp b/Engine/Source/Runtime/Foliage/Private/InstancedFoliage.cpp
--- a/Engine/Source/Runtime/Foliage/Private/InstancedFoliage.cpp
+++ b/Engine/Source/Runtime/Foliage/Private/InstancedFoliage.cpp
@@ -5787,6 +5787,20 @@ UFoliageInstancedStaticMeshComponent::UFoliageInstancedStaticMeshComponent(const
 	ViewRelevanceType = EHISMViewRelevanceType::Foliage;
 }
 
+#if WITH_EDITOR
+void UFoliageInstancedStaticMeshComponent::PreEditChange(FProperty* PropertyAboutToChange)
+{
+	// This component intentionally turns off the transaction flag when it is created and most changes go through the foliage editor.
+	//
+	// However, there are cases where users will modify the component directly and this allows us to capture those as proper transactions for undo/redo.  If a
+	// component level change is made on a property this will temporarily enable the transaction flag so it passes checks in transaction code.
+	//
+	SetFlags(RF_Transactional);
+	Super::PreEditChange(PropertyAboutToChange);
+	ClearFlags(RF_Transactional);
+}
+#endif
+
 void UFoliageInstancedStaticMeshComponent::ReceiveComponentDamage(float DamageAmount, FDamageEvent const& DamageEvent, AController* EventInstigator, AActor* DamageCauser)
 {
 	Super::ReceiveComponentDamage(DamageAmount, DamageEvent, EventInstigator, DamageCauser);
diff --git a/Engine/Source/Runtime/Foliage/Public/FoliageInstancedStaticMeshComponent.h b/Engine/Source/Runtime/Foliage/Public/FoliageInstancedStaticMeshComponent.h
--- a/Engine/Source/Runtime/Foliage/Public/FoliageInstancedStaticMeshComponent.h
+++ b/Engine/Source/Runtime/Foliage/Public/FoliageInstancedStaticMeshComponent.h
@@ -41,6 +41,7 @@ public:
 	//USceneComponent Interface
 #if WITH_EDITOR
 	FOLIAGE_API virtual uint64 GetHiddenEditorViews() const override;
+	FOLIAGE_API virtual void PreEditChange(FProperty* PropertyAboutToChange) override;
 #endif// WITH_EDITOR
 
 	/** Used by procedural generation to link generated component with its creator */
