Fix crash when using the foliage placement tool in maps with entities.

diff --git a/Engine/Source/Runtime/Foliage/Private/InstancedFoliage.cpp b/Engine/Source/Runtime/Foliage/Private/InstancedFoliage.cpp
--- a/Engine/Source/Runtime/Foliage/Private/InstancedFoliage.cpp
+++ b/Engine/Source/Runtime/Foliage/Private/InstancedFoliage.cpp
@@ -5262,20 +5262,22 @@ bool AInstancedFoliageActor::FoliageTrace(const UWorld* InWorld, FHitResult& Out
 		}
 
 		const UPrimitiveComponent* HitComponent = Hit.GetComponent();
-		check(HitComponent);
-
-		// In the editor traces can hit "No Collision" type actors, so ugh. (ignore these)
-		if (!HitComponent->IsQueryCollisionEnabled() || HitComponent->GetCollisionResponseToChannel(ECC_WorldStatic) != ECR_Block)
+		
+		if (HitComponent)
 		{
-			bOutDiscardHit = true;
-			return true;
-		}
+			// In the editor traces can hit "No Collision" type actors, so ugh. (ignore these)
+			if (!HitComponent->IsQueryCollisionEnabled() || HitComponent->GetCollisionResponseToChannel(ECC_WorldStatic) != ECR_Block)
+			{
+				bOutDiscardHit = true;
+				return true;
+			}
 
-		// Don't place foliage on invisible walls / triggers / volumes
-		if (HitComponent->IsA<UBrushComponent>())
-		{
-			bOutDiscardHit = true;
-			return true;
+			// Don't place foliage on invisible walls / triggers / volumes
+			if (HitComponent->IsA<UBrushComponent>())
+			{
+				bOutDiscardHit = true;
+				return true;
+			}
 		}
 
 		// Don't place foliage on itself
@@ -5284,6 +5286,8 @@ bool AInstancedFoliageActor::FoliageTrace(const UWorld* InWorld, FHitResult& Out
 
 		if (HitActor)
 		{
+			check(HitComponent);
+
 			// Don't place foliage on hidden actors.
 			if (HitActor->IsTemporarilyHiddenInEditor())
 			{
