PR #13896: Adding NanitePixelProgrammableDistance property to foliage type

diff --git a/Engine/Source/Runtime/Foliage/Private/InstancedFoliage.cpp b/Engine/Source/Runtime/Foliage/Private/InstancedFoliage.cpp
--- a/Engine/Source/Runtime/Foliage/Private/InstancedFoliage.cpp
+++ b/Engine/Source/Runtime/Foliage/Private/InstancedFoliage.cpp
@@ -632,6 +632,7 @@ UFoliageType::UFoliageType(const FObjectInitializer& ObjectInitializer)
 	bVisibleInRayTracing = true;
 	bEvaluateWorldPositionOffset = true;
 	WorldPositionOffsetDisableDistance = 0;
+	NanitePixelProgrammableDistance = 0.f;
 
 	ShadowCacheInvalidationBehavior = EShadowCacheInvalidationBehavior::Auto;
 
@@ -1733,6 +1734,12 @@ void FFoliageStaticMesh::UpdateComponentSettings(const UFoliageType_InstancedSta
 			bNeedsMarkRenderStateDirty = true;
 		}
 
+		if(Component->NanitePixelProgrammableDistance != FoliageType->NanitePixelProgrammableDistance)
+		{
+			Component->NanitePixelProgrammableDistance = FoliageType->NanitePixelProgrammableDistance;
+			bNeedsMarkRenderStateDirty = true;
+		}
+
 
 		if (Component->bEnableDensityScaling != FoliageType->bEnableDensityScaling)
 		{
diff --git a/Engine/Source/Runtime/Foliage/Public/FoliageType.h b/Engine/Source/Runtime/Foliage/Public/FoliageType.h
--- a/Engine/Source/Runtime/Foliage/Public/FoliageType.h
+++ b/Engine/Source/Runtime/Foliage/Public/FoliageType.h
@@ -363,6 +363,10 @@ public:
 	UPROPERTY(EditAnywhere, AdvancedDisplay, BlueprintReadOnly, Category = InstanceSettings, meta = (ClampMin=0))
 	int32 WorldPositionOffsetDisableDistance;
 
+	/** Used to forcefully disable pixel programmable rasterization of Nanite when the mesh is further than a given distance from the camera. */
+	UPROPERTY(EditAnywhere, AdvancedDisplay, BlueprintReadOnly, Category = InstanceSettings, meta = (ClampMin = 0))
+	float NanitePixelProgrammableDistance = 0.0f;
+
 	/** Custom collision for foliage */
 	UPROPERTY(EditAnywhere, Category=InstanceSettings, meta=(HideObjectType=true))
 	struct FBodyInstance BodyInstance;
